Fixed memory leak in FFoliageEdModeToolkit (TTP 336258)

diff --git a/Engine/Source/Editor/FoliageEdit/Private/SFoliageEdit.cpp b/Engine/Source/Editor/FoliageEdit/Private/SFoliageEdit.cpp
--- a/Engine/Source/Editor/FoliageEdit/Private/SFoliageEdit.cpp
+++ b/Engine/Source/Editor/FoliageEdit/Private/SFoliageEdit.cpp
@@ -19,11 +19,8 @@ class SFoliageDragDropHandler : public SCompoundWidget
 {
 public:
 	SLATE_BEGIN_ARGS(SFoliageDragDropHandler) {}
-
-	SLATE_DEFAULT_SLOT( FArguments, Content )
-
-	SLATE_ARGUMENT( TSharedPtr<SFoliageEdit>, FoliageEditPtr )
-
+	SLATE_DEFAULT_SLOT(FArguments, Content)
+	SLATE_ARGUMENT(TWeakPtr<SFoliageEdit>, FoliageEditPtr)
 	SLATE_END_ARGS()
 
 	/** SCompoundWidget functions */
@@ -41,24 +38,37 @@ public:
 
 	FReply OnDrop( const FGeometry& MyGeometry, const FDragDropEvent& DragDropEvent )
 	{
-		FoliageEditPtr->DisableDragDropOverlay();
-		return FoliageEditPtr->OnDrop_ListView(MyGeometry, DragDropEvent);
+		TSharedPtr<SFoliageEdit> PinnedPtr = FoliageEditPtr.Pin();
+		if (PinnedPtr.IsValid())
+		{
+			PinnedPtr->DisableDragDropOverlay();
+			return PinnedPtr->OnDrop_ListView(MyGeometry, DragDropEvent);
+		}
+		return FReply::Handled();
 	}
 
 	virtual void OnDragEnter( const FGeometry& MyGeometry, const FDragDropEvent& DragDropEvent ) OVERRIDE
 	{
-		FoliageEditPtr->EnableDragDropOverlay();
+		TSharedPtr<SFoliageEdit> PinnedPtr = FoliageEditPtr.Pin();
+		if (PinnedPtr.IsValid())
+		{
+			PinnedPtr->EnableDragDropOverlay();
+		}
 		SCompoundWidget::OnDragEnter(MyGeometry, DragDropEvent);
 	}
 
 	virtual void OnDragLeave( const FDragDropEvent& DragDropEvent ) OVERRIDE
 	{
-		FoliageEditPtr->DisableDragDropOverlay();
+		TSharedPtr<SFoliageEdit> PinnedPtr = FoliageEditPtr.Pin();
+		if (PinnedPtr.IsValid())
+		{
+			PinnedPtr->DisableDragDropOverlay();
+		}
 		SCompoundWidget::OnDragLeave(DragDropEvent);
 	}
 
 private:
-	TSharedPtr<SFoliageEdit> FoliageEditPtr;
+	TWeakPtr<SFoliageEdit> FoliageEditPtr;
 };
 
 BEGIN_SLATE_FUNCTION_BUILD_OPTIMIZATION
