Foliage editor refreshes UI when current level has been changed

diff --git a/Engine/Source/Editor/FoliageEdit/Private/FoliageEdMode.cpp b/Engine/Source/Editor/FoliageEdit/Private/FoliageEdMode.cpp
--- a/Engine/Source/Editor/FoliageEdit/Private/FoliageEdMode.cpp
+++ b/Engine/Source/Editor/FoliageEdit/Private/FoliageEdMode.cpp
@@ -94,6 +94,9 @@ void FEdModeFoliage::Enter()
 	// Load UI settings from config file
 	UISettings.Load();
 
+	// Bind to editor callbacks
+	FEditorDelegates::NewCurrentLevel.AddSP(this, &FEdModeFoliage::NotifyNewCurrentLevel);
+
 	// Force real-time viewports.  We'll back up the current viewport state so we can restore it when the
 	// user exits this mode.
 	const bool bWantRealTime = true;
@@ -141,6 +144,9 @@ void FEdModeFoliage::Exit()
 	FToolkitManager::Get().CloseToolkit(Toolkit.ToSharedRef());
 	Toolkit.Reset();
 
+	//
+	FEditorDelegates::NewCurrentLevel.RemoveAll(this);
+
 	// Remove the brush
 	SphereBrushComponent->UnregisterComponent();
 
diff --git a/Engine/Source/Editor/FoliageEdit/Private/FoliageEditModule.cpp b/Engine/Source/Editor/FoliageEdit/Private/FoliageEditModule.cpp
--- a/Engine/Source/Editor/FoliageEdit/Private/FoliageEditModule.cpp
+++ b/Engine/Source/Editor/FoliageEdit/Private/FoliageEditModule.cpp
@@ -30,7 +30,7 @@ public:
 		TSharedRef<FEdModeFoliage> NewEditorMode = MakeShareable(new FEdModeFoliage);
 		GEditorModeTools().RegisterMode(NewEditorMode);
 		EdModeFoliage = NewEditorMode;
-		FEditorDelegates::MapChange.AddSP( NewEditorMode, &FEdModeFoliage::NotifyMapRebuild );
+		FEditorDelegates::MapChange.AddSP(NewEditorMode, &FEdModeFoliage::NotifyMapRebuild);
 	}
 
 	/**
@@ -38,9 +38,10 @@ public:
 	 */
 	virtual void ShutdownModule() OVERRIDE
 	{
+		FEditorDelegates::MapChange.RemoveAll(EdModeFoliage.Get());
+
 		GEditorModeTools().UnregisterMode(EdModeFoliage.ToSharedRef());
 		EdModeFoliage = NULL;
-		FEditorDelegates::MapChange.RemoveAll(this);
 	}
 };
 
diff --git a/Engine/Source/Editor/FoliageEdit/Private/SFoliageEdit.cpp b/Engine/Source/Editor/FoliageEdit/Private/SFoliageEdit.cpp
--- a/Engine/Source/Editor/FoliageEdit/Private/SFoliageEdit.cpp
+++ b/Engine/Source/Editor/FoliageEdit/Private/SFoliageEdit.cpp
@@ -82,6 +82,8 @@ void SFoliageEdit::Construct(const FArguments& InArgs)
 
 	BindCommands();
 
+	FEditorDelegates::NewCurrentLevel.AddSP(this, &SFoliageEdit::NotifyNewCurrentLevel);
+
 	AssetThumbnailPool = MakeShareable( new FAssetThumbnailPool(512, TAttribute<bool>::Create( TAttribute<bool>::FGetter::CreateSP(this, &SFoliageEdit::IsHovered) ) ) );
 
 	bOverlayOverride = false;
@@ -317,6 +319,8 @@ SFoliageEdit::~SFoliageEdit()
 {
 	// Release all rendering resources being held onto
 	AssetThumbnailPool->ReleaseResources();
+	
+	FEditorDelegates::NewCurrentLevel.RemoveAll(this);
 }
 
 void SFoliageEdit::Tick( const FGeometry& AllottedGeometry, const double InCurrentTime, const float InDeltaTime )
@@ -720,4 +724,10 @@ void SFoliageEdit::EnableDragDropOverlay(void)
 	bOverlayOverride = true; 
 }
 
+void SFoliageEdit::NotifyNewCurrentLevel()
+{
+	RefreshFullList();
+}
+
+
 #undef LOCTEXT_NAMESPACE
diff --git a/Engine/Source/Editor/FoliageEdit/Private/SFoliageEdit.h b/Engine/Source/Editor/FoliageEdit/Private/SFoliageEdit.h
--- a/Engine/Source/Editor/FoliageEdit/Private/SFoliageEdit.h
+++ b/Engine/Source/Editor/FoliageEdit/Private/SFoliageEdit.h
@@ -165,6 +165,9 @@ private:
 	 */
 	bool CanAddStaticMesh(const UStaticMesh* const InStaticMesh) const;
 
+	/** Callback handler when user changes the active streaming level with the level browser */
+	void NotifyNewCurrentLevel();
+
 private:
 	/** The list view object for displaying Static Meshes to use for foliage. */
 	TSharedPtr< SListView< TSharedPtr< struct FFoliageMeshUIInfo > > > MeshListView;
