Improved Foliage painting perfomance

diff --git a/Engine/Source/Editor/FoliageEdit/Private/FoliageEdMode.cpp b/Engine/Source/Editor/FoliageEdit/Private/FoliageEdMode.cpp
--- a/Engine/Source/Editor/FoliageEdit/Private/FoliageEdMode.cpp
+++ b/Engine/Source/Editor/FoliageEdit/Private/FoliageEdMode.cpp
@@ -353,12 +353,8 @@ bool FEdModeFoliage::CapturedMouseMove( FEditorViewportClient* ViewportClient, F
 void FEdModeFoliage::GetRandomVectorInBrush( FVector& OutStart, FVector& OutEnd )
 {
 	// Find Rx and Ry inside the unit circle
-	float Ru, Rv;
-	do 
-	{
-		Ru = 2.f * FMath::FRand() - 1.f;
-		Rv = 2.f * FMath::FRand() - 1.f;
-	} while ( FMath::Square(Ru) + FMath::Square(Rv) > 1.f );
+	float Ru = (2.f * FMath::FRand() - 1.f);
+	float Rv = (2.f * FMath::FRand() - 1.f) * FMath::Sqrt(1.f - FMath::Square(Ru));
 	
 	// find random point in circle thru brush location parallel to screen surface
 	FVector U, V;
@@ -624,7 +620,11 @@ void FEdModeFoliage::AddInstancesForBrush( UWorld* InWorld, AInstancedFoliageAct
 
 		// We calculate a set of potential ExistingInstances for the brush area.
 		TArray<FPotentialInstance> PotentialInstanceBuckets[NUM_INSTANCE_BUCKETS];
-		FMemory::Memzero(PotentialInstanceBuckets, sizeof(PotentialInstanceBuckets));
+		// Reserve space in buckets for a potential instances 
+		for (auto& Bucket : PotentialInstanceBuckets)
+		{
+			Bucket.Reserve(DesiredInstanceCount);
+		}
 
 		// Quick lookup of potential instance locations, used for overlapping check.
 		TArray<FVector> PotentialInstanceLocations;
@@ -642,11 +642,12 @@ void FEdModeFoliage::AddInstancesForBrush( UWorld* InWorld, AInstancedFoliageAct
 			
 			FHitResult Hit;
 			static FName NAME_AddInstancesForBrush = FName(TEXT("AddInstancesForBrush"));
-			FCollisionQueryParams TraceParams(NAME_AddInstancesForBrush, true, AInstancedFoliageActor::GetInstancedFoliageActor(InWorld, false));
+			FCollisionQueryParams TraceParams(NAME_AddInstancesForBrush, true, IFA);
 			if( MeshSettings->VertexColorMask != FOLIAGEVERTEXCOLORMASK_Disabled )
 			{
 				TraceParams.bReturnFaceIndex = true;
 			}
+			
 			if( InWorld->LineTraceSingle(Hit, Start, End, TraceParams, FCollisionObjectQueryParams(ECC_WorldStatic))) 
 			{
 				// Check filters
@@ -706,7 +707,10 @@ void FEdModeFoliage::AddInstancesForBrush( UWorld* InWorld, AInstancedFoliageAct
 					}
 				}
 
-				new(PotentialInstanceBuckets[FMath::RoundToInt(HitWeight * (float)(NUM_INSTANCE_BUCKETS-1))]) FPotentialInstance( Hit.Location, Hit.Normal, Hit.Component.Get(), HitWeight );
+				const int32 BucketIndex = FMath::RoundToInt(HitWeight * (float)(NUM_INSTANCE_BUCKETS-1));
+				check(BucketIndex < ARRAY_COUNT(PotentialInstanceBuckets));
+
+				PotentialInstanceBuckets[BucketIndex].Add(FPotentialInstance( Hit.Location, Hit.Normal, Hit.Component.Get(), HitWeight ));
 			}
 		}
 
