Initial support for other foliage types

diff --git a/Engine/Source/Editor/FoliageEdit/Private/FoliageEdMode.cpp b/Engine/Source/Editor/FoliageEdit/Private/FoliageEdMode.cpp
--- a/Engine/Source/Editor/FoliageEdit/Private/FoliageEdMode.cpp
+++ b/Engine/Source/Editor/FoliageEdit/Private/FoliageEdMode.cpp
@@ -23,7 +23,7 @@
 
 // Classes
 #include "Foliage/InstancedFoliageActor.h"
-#include "Foliage/InstancedFoliageSettings.h"
+#include "Foliage/FoliageType.h"
 #include "Landscape/LandscapeComponent.h"
 #include "Landscape/LandscapeHeightfieldCollisionComponent.h"
 #include "Landscape/LandscapeInfo.h"
@@ -35,19 +35,19 @@
 //
 
 /** Constructor */
-FEdModeFoliage::FEdModeFoliage() 
-:	FEdMode()
-,	bToolActive(false)
-,	SelectionIFA(NULL)
-,	bCanAltDrag(false)
+FEdModeFoliage::FEdModeFoliage()
+	: FEdMode()
+	, bToolActive(false)
+	, SelectionIFA(nullptr)
+	, bCanAltDrag(false)
 {
 	// Load resources and construct brush component
-	UMaterial* BrushMaterial = NULL;
-	UStaticMesh* StaticMesh = NULL;
+	UMaterial* BrushMaterial = nullptr;
+	UStaticMesh* StaticMesh = nullptr;
 	if (!IsRunningCommandlet())
 	{
-		BrushMaterial = LoadObject<UMaterial>(NULL, TEXT("/Engine/EditorLandscapeResources/FoliageBrushSphereMaterial.FoliageBrushSphereMaterial"), NULL, LOAD_None, NULL);
-		StaticMesh = LoadObject<UStaticMesh>(NULL,TEXT("/Engine/EngineMeshes/Sphere.Sphere"), NULL, LOAD_None, NULL);
+		BrushMaterial = LoadObject<UMaterial>(nullptr, TEXT("/Engine/EditorLandscapeResources/FoliageBrushSphereMaterial.FoliageBrushSphereMaterial"), nullptr, LOAD_None, nullptr);
+		StaticMesh = LoadObject<UStaticMesh>(nullptr, TEXT("/Engine/EngineMeshes/Sphere.Sphere"), nullptr, LOAD_None, nullptr);
 	}
 
 	SphereBrushComponent = ConstructObject<UStaticMeshComponent>(UStaticMeshComponent::StaticClass());
@@ -70,10 +70,10 @@ FEdModeFoliage::~FEdModeFoliage()
 
 
 /** FGCObject interface */
-void FEdModeFoliage::AddReferencedObjects( FReferenceCollector& Collector )
+void FEdModeFoliage::AddReferencedObjects(FReferenceCollector& Collector)
 {
 	// Call parent implementation
-	FEdMode::AddReferencedObjects( Collector );
+	FEdMode::AddReferencedObjects(Collector);
 
 	Collector.AddReferencedObject(SphereBrushComponent);
 }
@@ -81,7 +81,7 @@ void FEdModeFoliage::AddReferencedObjects( FReferenceCollector& Collector )
 /** FEdMode: Called when the mode is created */
 void FEdModeFoliage::Initialize()
 {
-	FEditorDelegates::MapChange.AddSP( this, &FEdModeFoliage::NotifyMapRebuild );
+	FEditorDelegates::MapChange.AddSP(this, &FEdModeFoliage::NotifyMapRebuild);
 }
 
 /** FEdMode: Called when the mode is entered */
@@ -102,35 +102,35 @@ void FEdModeFoliage::Enter()
 	// user exits this mode.
 	const bool bWantRealTime = true;
 	const bool bRememberCurrentState = true;
-	ForceRealTimeViewports( bWantRealTime, bRememberCurrentState );
+	ForceRealTimeViewports(bWantRealTime, bRememberCurrentState);
 
 	// Reapply selection visualization on any foliage items
-	if( UISettings.GetSelectToolSelected() || UISettings.GetLassoSelectToolSelected() )
+	if (UISettings.GetSelectToolSelected() || UISettings.GetLassoSelectToolSelected())
 	{
-		SelectionIFA = AInstancedFoliageActor::GetInstancedFoliageActor(GetWorld());
+		SelectionIFA = AInstancedFoliageActor::GetInstancedFoliageActorForCurrentLevel(GetWorld());
 		SelectionIFA->ApplySelectionToComponents(true);
 	}
 	else
 	{
-		SelectionIFA = NULL;
+		SelectionIFA = nullptr;
 	}
 
 	if (!Toolkit.IsValid())
 	{
 		// @todo: Remove this assumption when we make modes per level editor instead of global
-		auto ToolkitHost = FModuleManager::LoadModuleChecked< FLevelEditorModule >( "LevelEditor" ).GetFirstLevelEditor();
+		auto ToolkitHost = FModuleManager::LoadModuleChecked< FLevelEditorModule >("LevelEditor").GetFirstLevelEditor();
 		Toolkit = MakeShareable(new FFoliageEdModeToolkit);
 		Toolkit->Init(ToolkitHost);
 	}
 
 	// Fixup any broken clusters
-	AInstancedFoliageActor* IFA = AInstancedFoliageActor::GetInstancedFoliageActor(GetWorld());
+	AInstancedFoliageActor* IFA = AInstancedFoliageActor::GetInstancedFoliageActorForCurrentLevel(GetWorld());
 	for (auto& MeshPair : IFA->FoliageMeshes)
 	{
 		FFoliageMeshInfo& MeshInfo = *MeshPair.Value;
-		for( int32 ClusterIdx=0;ClusterIdx<MeshInfo.InstanceClusters.Num();ClusterIdx++ )
+		for (FFoliageInstanceCluster& Cluster : MeshInfo.InstanceClusters)
 		{
-			if( MeshInfo.InstanceClusters[ClusterIdx].ClusterComponent == NULL )
+			if (Cluster.ClusterComponent == nullptr)
 			{
 				MeshInfo.ReallocateClusters(IFA, MeshPair.Key);
 				break;
@@ -154,7 +154,7 @@ void FEdModeFoliage::Exit()
 	// Restore real-time viewport state if we changed it
 	const bool bWantRealTime = false;
 	const bool bRememberCurrentState = false;
-	ForceRealTimeViewports( bWantRealTime, bRememberCurrentState );
+	ForceRealTimeViewports(bWantRealTime, bRememberCurrentState);
 
 	// Clear the cache (safety, should be empty!)
 	LandscapeLayerCaches.Empty();
@@ -166,7 +166,7 @@ void FEdModeFoliage::Exit()
 	FoliageMeshList.Empty();
 
 	// Clear selection visualization on any foliage items
-	if( SelectionIFA && (UISettings.GetSelectToolSelected() || UISettings.GetLassoSelectToolSelected()) )
+	if (SelectionIFA && (UISettings.GetSelectToolSelected() || UISettings.GetLassoSelectToolSelected()))
 	{
 		SelectionIFA->ApplySelectionToComponents(false);
 	}
@@ -186,13 +186,13 @@ void FEdModeFoliage::PostUndo()
 void FEdModeFoliage::NotifyNewCurrentLevel()
 {
 	// Remove any selections in the old level and reapply for the new level
-	if( UISettings.GetSelectToolSelected() || UISettings.GetLassoSelectToolSelected() )
+	if (UISettings.GetSelectToolSelected() || UISettings.GetLassoSelectToolSelected())
 	{
-		if( SelectionIFA )
+		if (SelectionIFA)
 		{
 			SelectionIFA->ApplySelectionToComponents(false);
 		}
-		SelectionIFA = AInstancedFoliageActor::GetInstancedFoliageActor(GetWorld());
+		SelectionIFA = AInstancedFoliageActor::GetInstancedFoliageActorForCurrentLevel(GetWorld());
 		SelectionIFA->ApplySelectionToComponents(true);
 	}
 }
@@ -200,30 +200,30 @@ void FEdModeFoliage::NotifyNewCurrentLevel()
 /** When the user changes the current tool in the UI */
 void FEdModeFoliage::NotifyToolChanged()
 {
-	if( UISettings.GetSelectToolSelected() || UISettings.GetLassoSelectToolSelected() )
+	if (UISettings.GetSelectToolSelected() || UISettings.GetLassoSelectToolSelected())
 	{
-		if( SelectionIFA == NULL )
+		if (SelectionIFA == nullptr)
 		{
-			SelectionIFA = AInstancedFoliageActor::GetInstancedFoliageActor(GetWorld());
+			SelectionIFA = AInstancedFoliageActor::GetInstancedFoliageActorForCurrentLevel(GetWorld());
 			SelectionIFA->ApplySelectionToComponents(true);
 		}
 	}
 	else
 	{
-		if( SelectionIFA )
+		if (SelectionIFA)
 		{
 			SelectionIFA->ApplySelectionToComponents(false);
 		}
-		SelectionIFA = NULL;
+		SelectionIFA = nullptr;
 	}
 }
 
 void FEdModeFoliage::NotifyMapRebuild(uint32 MapChangeFlags) const
 {
-	if(MapChangeEventFlags::MapRebuild & MapChangeFlags)
+	if (MapChangeEventFlags::MapRebuild & MapChangeFlags)
 	{
-		AInstancedFoliageActor* IFA = AInstancedFoliageActor::GetInstancedFoliageActor(GetWorld());
-		if( IFA )
+		AInstancedFoliageActor* IFA = AInstancedFoliageActor::GetInstancedFoliageActorForCurrentLevel(GetWorld());
+		if (IFA)
 		{
 			IFA->MapRebuild();
 		}
@@ -237,37 +237,37 @@ bool FEdModeFoliage::DisallowMouseDeltaTracking() const
 }
 
 /** FEdMode: Called once per frame */
-void FEdModeFoliage::Tick(FEditorViewportClient* ViewportClient,float DeltaTime)
+void FEdModeFoliage::Tick(FEditorViewportClient* ViewportClient, float DeltaTime)
 {
-	if( IsCtrlDown(ViewportClient->Viewport) && bToolActive )
+	if (IsCtrlDown(ViewportClient->Viewport) && bToolActive)
 	{
 		ApplyBrush(ViewportClient);
 	}
 
-	FEdMode::Tick(ViewportClient,DeltaTime);
+	FEdMode::Tick(ViewportClient, DeltaTime);
 
-	if( UISettings.GetSelectToolSelected() || UISettings.GetLassoSelectToolSelected() )
+	if (UISettings.GetSelectToolSelected() || UISettings.GetLassoSelectToolSelected())
 	{
 		// Update pivot
-		AInstancedFoliageActor* IFA = AInstancedFoliageActor::GetInstancedFoliageActor(GetWorld());
+		AInstancedFoliageActor* IFA = AInstancedFoliageActor::GetInstancedFoliageActorForCurrentLevel(GetWorld());
 		Owner->PivotLocation = Owner->SnappedLocation = IFA->GetSelectionLocation();
 	}
-	
+
 	// Update the position and size of the brush component
-	if( bBrushTraceValid && (UISettings.GetPaintToolSelected() || UISettings.GetReapplyToolSelected() || UISettings.GetLassoSelectToolSelected()) )
+	if (bBrushTraceValid && (UISettings.GetPaintToolSelected() || UISettings.GetReapplyToolSelected() || UISettings.GetLassoSelectToolSelected()))
 	{
 		// Scale adjustment is due to default sphere SM size.
-		FTransform BrushTransform = FTransform( FQuat::Identity, BrushLocation, FVector(UISettings.GetRadius() * 0.00625f) );
+		FTransform BrushTransform = FTransform(FQuat::Identity, BrushLocation, FVector(UISettings.GetRadius() * 0.00625f));
 		SphereBrushComponent->SetRelativeTransform(BrushTransform);
 
-		if( !SphereBrushComponent->IsRegistered() )
+		if (!SphereBrushComponent->IsRegistered())
 		{
 			SphereBrushComponent->RegisterComponentWithWorld(ViewportClient->GetWorld());
 		}
 	}
 	else
 	{
-		if( SphereBrushComponent->IsRegistered() )
+		if (SphereBrushComponent->IsRegistered())
 		{
 			SphereBrushComponent->UnregisterComponent();
 		}
@@ -275,19 +275,19 @@ void FEdModeFoliage::Tick(FEditorViewportClient* ViewportClient,float DeltaTime)
 }
 
 /** Trace under the mouse cursor and update brush position */
-void FEdModeFoliage::FoliageBrushTrace( FEditorViewportClient* ViewportClient, int32 MouseX, int32 MouseY )
+void FEdModeFoliage::FoliageBrushTrace(FEditorViewportClient* ViewportClient, int32 MouseX, int32 MouseY)
 {
 	bBrushTraceValid = false;
-	if( UISettings.GetPaintToolSelected() || UISettings.GetReapplyToolSelected() || UISettings.GetLassoSelectToolSelected() )
+	if (UISettings.GetPaintToolSelected() || UISettings.GetReapplyToolSelected() || UISettings.GetLassoSelectToolSelected())
 	{
 		// Compute a world space ray from the screen space mouse coordinates
-		FSceneViewFamilyContext ViewFamily( FSceneViewFamily::ConstructionValues(
-			ViewportClient->Viewport, 
+		FSceneViewFamilyContext ViewFamily(FSceneViewFamily::ConstructionValues(
+			ViewportClient->Viewport,
 			ViewportClient->GetScene(),
-			ViewportClient->EngineShowFlags )
-			.SetRealtimeUpdate( ViewportClient->IsRealtime() ));
-		FSceneView* View = ViewportClient->CalcSceneView( &ViewFamily );
-		FViewportCursorLocation MouseViewportRay( View, ViewportClient, MouseX, MouseY );
+			ViewportClient->EngineShowFlags)
+			.SetRealtimeUpdate(ViewportClient->IsRealtime()));
+		FSceneView* View = ViewportClient->CalcSceneView(&ViewFamily);
+		FViewportCursorLocation MouseViewportRay(View, ViewportClient, MouseX, MouseY);
 
 		FVector Start = MouseViewportRay.GetOrigin();
 		BrushTraceDirection = MouseViewportRay.GetDirection();
@@ -295,7 +295,7 @@ void FEdModeFoliage::FoliageBrushTrace( FEditorViewportClient* ViewportClient, i
 
 		FHitResult Hit;
 		UWorld* World = ViewportClient->GetWorld();
-		if ( World->LineTraceSingle(Hit, Start, End, FCollisionQueryParams(NAME_None, true, AInstancedFoliageActor::GetInstancedFoliageActor(World, false)), FCollisionObjectQueryParams(ECC_WorldStatic)) )
+		if (World->LineTraceSingle(Hit, Start, End, FCollisionQueryParams(NAME_None, true, AInstancedFoliageActor::GetInstancedFoliageActorForCurrentLevel(World, false)), FCollisionObjectQueryParams(ECC_WorldStatic)))
 		{
 			// Check filters
 			UPrimitiveComponent* PrimComp = Hit.Component.Get();
@@ -328,9 +328,9 @@ void FEdModeFoliage::FoliageBrushTrace( FEditorViewportClient* ViewportClient, i
  *
  * @return	true if input was handled
  */
-bool FEdModeFoliage::MouseMove( FEditorViewportClient* ViewportClient, FViewport* Viewport, int32 MouseX, int32 MouseY )
+bool FEdModeFoliage::MouseMove(FEditorViewportClient* ViewportClient, FViewport* Viewport, int32 MouseX, int32 MouseY)
 {
-	FoliageBrushTrace(ViewportClient, MouseX, MouseY );
+	FoliageBrushTrace(ViewportClient, MouseX, MouseY);
 	return false;
 }
 
@@ -344,58 +344,58 @@ bool FEdModeFoliage::MouseMove( FEditorViewportClient* ViewportClient, FViewport
  *
  * @return	true if input was handled
  */
-bool FEdModeFoliage::CapturedMouseMove( FEditorViewportClient* ViewportClient, FViewport* Viewport, int32 MouseX, int32 MouseY )
+bool FEdModeFoliage::CapturedMouseMove(FEditorViewportClient* ViewportClient, FViewport* Viewport, int32 MouseX, int32 MouseY)
 {
-	FoliageBrushTrace(ViewportClient, MouseX, MouseY );
+	FoliageBrushTrace(ViewportClient, MouseX, MouseY);
 	return false;
 }
 
-void FEdModeFoliage::GetRandomVectorInBrush( FVector& OutStart, FVector& OutEnd )
+void FEdModeFoliage::GetRandomVectorInBrush(FVector& OutStart, FVector& OutEnd)
 {
 	// Find Rx and Ry inside the unit circle
 	float Ru = (2.f * FMath::FRand() - 1.f);
 	float Rv = (2.f * FMath::FRand() - 1.f) * FMath::Sqrt(1.f - FMath::Square(Ru));
-	
+
 	// find random point in circle thru brush location parallel to screen surface
 	FVector U, V;
-	BrushTraceDirection.FindBestAxisVectors(U,V);
+	BrushTraceDirection.FindBestAxisVectors(U, V);
 	FVector Point = Ru * U + Rv * V;
 
 	// find distance to surface of sphere brush from this point
-	float Rw = FMath::Sqrt( 1.f - (FMath::Square(Ru) + FMath::Square(Rv)) );
+	float Rw = FMath::Sqrt(1.f - (FMath::Square(Ru) + FMath::Square(Rv)));
 
 	OutStart = BrushLocation + UISettings.GetRadius() * (Ru * U + Rv * V - Rw * BrushTraceDirection);
-	OutEnd   = BrushLocation + UISettings.GetRadius() * (Ru * U + Rv * V + Rw * BrushTraceDirection);
+	OutEnd = BrushLocation + UISettings.GetRadius() * (Ru * U + Rv * V + Rw * BrushTraceDirection);
 }
 
 
 // Number of buckets for layer weight histogram distribution.
 #define NUM_INSTANCE_BUCKETS 10
 
-bool CheckCollisionWithWorld(UInstancedFoliageSettings* MeshSettings, const FFoliageInstance& Inst, const FVector& HitNormal, const FVector& HitLocation,UWorld* InWorld)
+bool CheckCollisionWithWorld(UFoliageType* Settings, const FFoliageInstance& Inst, const FVector& HitNormal, const FVector& HitLocation, UWorld* InWorld)
 {
 	FMatrix InstTransform = Inst.GetInstanceWorldTransform().ToMatrixWithScale();
 	FVector LocalHit = InstTransform.InverseTransformPosition(HitLocation);
 
-	if ( MeshSettings->CollisionWithWorld )
+	if (Settings->CollisionWithWorld)
 	{
 		// Check for overhanging ledge
 		{
-			FVector LocalSamplePos[4] = {	FVector(MeshSettings->LowBoundOriginRadius.Z, 0, 0), 
-				FVector(-MeshSettings->LowBoundOriginRadius.Z, 0, 0),
-				FVector(0, MeshSettings->LowBoundOriginRadius.Z, 0),
-				FVector(0, -MeshSettings->LowBoundOriginRadius.Z, 0)
+			FVector LocalSamplePos[4] = { FVector(Settings->LowBoundOriginRadius.Z, 0, 0),
+				FVector(-Settings->LowBoundOriginRadius.Z, 0, 0),
+				FVector(0, Settings->LowBoundOriginRadius.Z, 0),
+				FVector(0, -Settings->LowBoundOriginRadius.Z, 0)
 			};
 
 			for (uint32 i = 0; i < 4; ++i)
 			{
 				FHitResult Hit;
-				FVector SamplePos = InstTransform.TransformPosition(FVector(MeshSettings->LowBoundOriginRadius.X, MeshSettings->LowBoundOriginRadius.Y, 2.f) + LocalSamplePos[i]);
-				float WorldRadius = (MeshSettings->LowBoundOriginRadius.Z + 2.f)*FMath::Max(Inst.DrawScale3D.X, Inst.DrawScale3D.Y);
-				FVector NormalVector = MeshSettings->AlignToNormal ? HitNormal : FVector(0, 0, 1);
-				if ( InWorld->LineTraceSingle(Hit, SamplePos, SamplePos - NormalVector*WorldRadius, FCollisionQueryParams(true), FCollisionObjectQueryParams(ECC_WorldStatic)) )
+				FVector SamplePos = InstTransform.TransformPosition(FVector(Settings->LowBoundOriginRadius.X, Settings->LowBoundOriginRadius.Y, 2.f) + LocalSamplePos[i]);
+				float WorldRadius = (Settings->LowBoundOriginRadius.Z + 2.f)*FMath::Max(Inst.DrawScale3D.X, Inst.DrawScale3D.Y);
+				FVector NormalVector = Settings->AlignToNormal ? HitNormal : FVector(0, 0, 1);
+				if (InWorld->LineTraceSingle(Hit, SamplePos, SamplePos - NormalVector*WorldRadius, FCollisionQueryParams(true), FCollisionObjectQueryParams(ECC_WorldStatic)))
 				{
-					if (LocalHit.Z - Inst.ZOffset < MeshSettings->LowBoundOriginRadius.Z)
+					if (LocalHit.Z - Inst.ZOffset < Settings->LowBoundOriginRadius.Z)
 					{
 						continue;
 					}
@@ -406,13 +406,13 @@ bool CheckCollisionWithWorld(UInstancedFoliageSettings* MeshSettings, const FFol
 
 		// Check collision with Bounding Box
 		{
-			FBox MeshBox = MeshSettings->MeshBounds.GetBox();
-			MeshBox.Min.Z = FMath::Min(MeshBox.Max.Z, LocalHit.Z + MeshSettings->MeshBounds.BoxExtent.Z * 0.05f);
+			FBox MeshBox = Settings->MeshBounds.GetBox();
+			MeshBox.Min.Z = FMath::Min(MeshBox.Max.Z, LocalHit.Z + Settings->MeshBounds.BoxExtent.Z * 0.05f);
 			FBoxSphereBounds ShrinkBound(MeshBox);
 			FBoxSphereBounds WorldBound = ShrinkBound.TransformBy(InstTransform);
 			//::DrawDebugBox(InWorld, WorldBound.Origin, WorldBound.BoxExtent, FColor(255, 0, 0), true, 10.f);
 			static FName NAME_FoliageCollisionWithWorld = FName(TEXT("FoliageCollisionWithWorld"));
-			if ( InWorld->OverlapTest(WorldBound.Origin, FQuat(Inst.Rotation), FCollisionShape::MakeBox(ShrinkBound.BoxExtent * Inst.DrawScale3D * MeshSettings->CollisionScale), FCollisionQueryParams(NAME_FoliageCollisionWithWorld, false), FCollisionObjectQueryParams(ECC_WorldStatic) ) )
+			if (InWorld->OverlapTest(WorldBound.Origin, FQuat(Inst.Rotation), FCollisionShape::MakeBox(ShrinkBound.BoxExtent * Inst.DrawScale3D * Settings->CollisionScale), FCollisionQueryParams(NAME_FoliageCollisionWithWorld, false), FCollisionObjectQueryParams(ECC_WorldStatic)))
 			{
 				return false;
 			}
@@ -429,37 +429,37 @@ struct FPotentialInstance
 	FVector HitNormal;
 	UPrimitiveComponent* HitComponent;
 	float HitWeight;
-	
+
 	FPotentialInstance(FVector InHitLocation, FVector InHitNormal, UPrimitiveComponent* InHitComponent, float InHitWeight)
-	:	HitLocation(InHitLocation)
-	,	HitNormal(InHitNormal)
-	,	HitComponent(InHitComponent)
-	,	HitWeight(InHitWeight)
+		: HitLocation(InHitLocation)
+		, HitNormal(InHitNormal)
+		, HitComponent(InHitComponent)
+		, HitWeight(InHitWeight)
 	{}
 
-	bool PlaceInstance(UInstancedFoliageSettings* MeshSettings, FFoliageInstance& Inst,UWorld* InWorld)
+	bool PlaceInstance(UFoliageType* Settings, FFoliageInstance& Inst, UWorld* InWorld)
 	{
-		if( MeshSettings->UniformScale )
+		if (Settings->UniformScale)
 		{
-			float Scale = MeshSettings->ScaleMinX + FMath::FRand() * (MeshSettings->ScaleMaxX - MeshSettings->ScaleMinX);
-			Inst.DrawScale3D = FVector(Scale,Scale,Scale);
+			float Scale = Settings->ScaleMinX + FMath::FRand() * (Settings->ScaleMaxX - Settings->ScaleMinX);
+			Inst.DrawScale3D = FVector(Scale, Scale, Scale);
 		}
 		else
 		{
 			float LockRand = FMath::FRand();
-			Inst.DrawScale3D.X = MeshSettings->ScaleMinX + (MeshSettings->LockScaleX ? LockRand : FMath::FRand()) * (MeshSettings->ScaleMaxX - MeshSettings->ScaleMinX);
-			Inst.DrawScale3D.Y = MeshSettings->ScaleMinY + (MeshSettings->LockScaleY ? LockRand : FMath::FRand()) * (MeshSettings->ScaleMaxY - MeshSettings->ScaleMinY);
-			Inst.DrawScale3D.Z = MeshSettings->ScaleMinZ + (MeshSettings->LockScaleZ ? LockRand : FMath::FRand()) * (MeshSettings->ScaleMaxZ - MeshSettings->ScaleMinZ);
+			Inst.DrawScale3D.X = Settings->ScaleMinX + (Settings->LockScaleX ? LockRand : FMath::FRand()) * (Settings->ScaleMaxX - Settings->ScaleMinX);
+			Inst.DrawScale3D.Y = Settings->ScaleMinY + (Settings->LockScaleY ? LockRand : FMath::FRand()) * (Settings->ScaleMaxY - Settings->ScaleMinY);
+			Inst.DrawScale3D.Z = Settings->ScaleMinZ + (Settings->LockScaleZ ? LockRand : FMath::FRand()) * (Settings->ScaleMaxZ - Settings->ScaleMinZ);
 		}
 
-		Inst.ZOffset = MeshSettings->ZOffsetMin + FMath::FRand() * (MeshSettings->ZOffsetMax - MeshSettings->ZOffsetMin);
+		Inst.ZOffset = Settings->ZOffsetMin + FMath::FRand() * (Settings->ZOffsetMax - Settings->ZOffsetMin);
 
 		Inst.Location = HitLocation;
 
 		// Random yaw and optional random pitch up to the maximum
-		Inst.Rotation = FRotator(FMath::FRand() * MeshSettings->RandomPitchAngle,0.f,0.f);
+		Inst.Rotation = FRotator(FMath::FRand() * Settings->RandomPitchAngle, 0.f, 0.f);
 
-		if( MeshSettings->RandomYaw )
+		if (Settings->RandomYaw)
 		{
 			Inst.Rotation.Yaw = FMath::FRand() * 360.f;
 		}
@@ -468,76 +468,75 @@ struct FPotentialInstance
 			Inst.Flags |= FOLIAGE_NoRandomYaw;
 		}
 
-		if( MeshSettings->AlignToNormal )
+		if (Settings->AlignToNormal)
 		{
-			Inst.AlignToNormal( HitNormal, MeshSettings->AlignMaxAngle );
+			Inst.AlignToNormal(HitNormal, Settings->AlignMaxAngle);
 		}
 
 		// Apply the Z offset in local space
-		if( FMath::Abs(Inst.ZOffset) > KINDA_SMALL_NUMBER )
+		if (FMath::Abs(Inst.ZOffset) > KINDA_SMALL_NUMBER)
 		{
-			Inst.Location = Inst.GetInstanceWorldTransform().TransformPosition(FVector(0,0,Inst.ZOffset));
+			Inst.Location = Inst.GetInstanceWorldTransform().TransformPosition(FVector(0, 0, Inst.ZOffset));
 		}
 
 		Inst.Base = HitComponent;
 
-		return CheckCollisionWithWorld(MeshSettings, Inst, HitNormal, HitLocation,InWorld);
+		return CheckCollisionWithWorld(Settings, Inst, HitNormal, HitLocation, InWorld);
 	}
 };
 
-static bool CheckLocationForPotentialInstance(FFoliageMeshInfo& MeshInfo, const UInstancedFoliageSettings* MeshSettings, float DensityCheckRadius, const FVector& Location, const FVector& Normal, TArray<FVector>& PotentialInstanceLocations, FFoliageInstanceHash& PotentialInstanceHash )
+static bool CheckLocationForPotentialInstance(FFoliageMeshInfo& MeshInfo, const UFoliageType* Settings, float DensityCheckRadius, const FVector& Location, const FVector& Normal, TArray<FVector>& PotentialInstanceLocations, FFoliageInstanceHash& PotentialInstanceHash)
 {
 	// Check slope
-	if( (MeshSettings->GroundSlope > 0.f && Normal.Z <= FMath::Cos( PI * MeshSettings->GroundSlope / 180.f)-SMALL_NUMBER) ||
-		(MeshSettings->GroundSlope < 0.f && Normal.Z >= FMath::Cos(-PI * MeshSettings->GroundSlope / 180.f)+SMALL_NUMBER) )
+	if ((Settings->GroundSlope > 0.f && Normal.Z <= FMath::Cos(PI * Settings->GroundSlope / 180.f) - SMALL_NUMBER) ||
+		(Settings->GroundSlope < 0.f && Normal.Z >= FMath::Cos(-PI * Settings->GroundSlope / 180.f) + SMALL_NUMBER))
 	{
 		return false;
 	}
 
 	// Check height range
-	if( Location.Z < MeshSettings->HeightMin || Location.Z > MeshSettings->HeightMax )
+	if (Location.Z < Settings->HeightMin || Location.Z > Settings->HeightMax)
 	{
 		return false;
 	}
 
 	// Check existing instances. Use the Density radius rather than the minimum radius
-	if( MeshInfo.CheckForOverlappingSphere(FSphere(Location, DensityCheckRadius)) )
+	if (MeshInfo.CheckForOverlappingSphere(FSphere(Location, DensityCheckRadius)))
 	{
 		return false;
 	}
 
 	// Check if we're too close to any other instances
-	if( MeshSettings->Radius > 0.f )
+	if (Settings->Radius > 0.f)
 	{
 		// Check with other potential instances we're about to add.
 		bool bFoundOverlapping = false;
-		float RadiusSquared = FMath::Square(DensityCheckRadius/*MeshSettings->Radius*/);
+		float RadiusSquared = FMath::Square(DensityCheckRadius/*Settings->Radius*/);
 
-		TSet<int32> TempInstances;
-		PotentialInstanceHash.GetInstancesOverlappingBox(FBox::BuildAABB(Location, FVector(MeshSettings->Radius,MeshSettings->Radius,MeshSettings->Radius)), TempInstances );
-		for( TSet<int32>::TConstIterator It(TempInstances); It; ++It )
+		auto TempInstances = PotentialInstanceHash.GetInstancesOverlappingBox(FBox::BuildAABB(Location, FVector(Settings->Radius)));
+		for (int32 Idx : TempInstances)
 		{
-			if( (PotentialInstanceLocations[*It] - Location).SizeSquared() < RadiusSquared )
+			if ((PotentialInstanceLocations[Idx] - Location).SizeSquared() < RadiusSquared)
 			{
 				bFoundOverlapping = true;
 				break;
 			}
 		}
-		if( bFoundOverlapping )
+		if (bFoundOverlapping)
 		{
 			return false;
 		}
 	}
 
 	int32 PotentialIdx = PotentialInstanceLocations.Add(Location);
-	PotentialInstanceHash.InsertInstance( Location, PotentialIdx );
+	PotentialInstanceHash.InsertInstance(Location, PotentialIdx);
 	return true;
 }
 
-static bool CheckVertexColor(const UInstancedFoliageSettings* MeshSettings, const FColor& VertexColor)
+static bool CheckVertexColor(const UFoliageType* Settings, const FColor& VertexColor)
 {
 	uint8 ColorChannel;
-	switch( MeshSettings->VertexColorMask )
+	switch (Settings->VertexColorMask)
 	{
 	case FOLIAGEVERTEXCOLORMASK_Red:
 		ColorChannel = VertexColor.R;
@@ -555,16 +554,16 @@ static bool CheckVertexColor(const UInstancedFoliageSettings* MeshSettings, cons
 		return true;
 	}
 
-	if( MeshSettings->VertexColorMaskInvert )
+	if (Settings->VertexColorMaskInvert)
 	{
-		if( ColorChannel > FMath::RoundToInt(MeshSettings->VertexColorMaskThreshold * 255.f) )
+		if (ColorChannel > FMath::RoundToInt(Settings->VertexColorMaskThreshold * 255.f))
 		{
 			return false;
 		}
 	}
 	else
 	{
-		if( ColorChannel < FMath::RoundToInt(MeshSettings->VertexColorMaskThreshold * 255.f) )
+		if (ColorChannel < FMath::RoundToInt(Settings->VertexColorMaskThreshold * 255.f))
 		{
 			return false;
 		}
@@ -575,39 +574,38 @@ static bool CheckVertexColor(const UInstancedFoliageSettings* MeshSettings, cons
 
 
 /** Add instances inside the brush to match DesiredInstanceCount */
-void FEdModeFoliage::AddInstancesForBrush( UWorld* InWorld, AInstancedFoliageActor* IFA, UStaticMesh* StaticMesh, FFoliageMeshInfo& MeshInfo, int32 DesiredInstanceCount, TArray<int32>& ExistingInstances, float Pressure )
+void FEdModeFoliage::AddInstancesForBrush(UWorld* InWorld, AInstancedFoliageActor* IFA, UFoliageType* Settings, FFoliageMeshInfo& MeshInfo, int32 DesiredInstanceCount, TArray<int32>& ExistingInstances, float Pressure)
 {
-	checkf( InWorld == IFA->GetWorld() , TEXT("Warning:World does not match Foliage world") );
-	UInstancedFoliageSettings* MeshSettings = MeshInfo.Settings;
-	if( DesiredInstanceCount > ExistingInstances.Num() )
+	checkf(InWorld == IFA->GetWorld(), TEXT("Warning:World does not match Foliage world"));
+	if (DesiredInstanceCount > ExistingInstances.Num())
 	{
 		int32 ExistingInstanceBuckets[NUM_INSTANCE_BUCKETS];
 		FMemory::Memzero(ExistingInstanceBuckets, sizeof(ExistingInstanceBuckets));
 
 		// Cache store mapping between component and weight data
-		TMap<ULandscapeComponent*, TArray<uint8> >* LandscapeLayerCache = NULL;
+		TMap<ULandscapeComponent*, TArray<uint8> >* LandscapeLayerCache = nullptr;
 
-		FName LandscapeLayerName = MeshSettings->LandscapeLayer;
-		if( LandscapeLayerName != NAME_None )
+		FName LandscapeLayerName = Settings->LandscapeLayer;
+		if (LandscapeLayerName != NAME_None)
 		{
 			LandscapeLayerCache = &LandscapeLayerCaches.FindOrAdd(LandscapeLayerName);
 
 			// Find the landscape weights of existing ExistingInstances
-			for( int32 Idx=0;Idx<ExistingInstances.Num();Idx++ )
+			for (int32 Idx : ExistingInstances)
 			{
-				FFoliageInstance& Instance = MeshInfo.Instances[ExistingInstances[Idx]];
+				FFoliageInstance& Instance = MeshInfo.Instances[Idx];
 				ULandscapeHeightfieldCollisionComponent* HitLandscapeCollision = Cast<ULandscapeHeightfieldCollisionComponent>(Instance.Base);
-				if( HitLandscapeCollision )
+				if (HitLandscapeCollision)
 				{
 					ULandscapeComponent* HitLandscape = HitLandscapeCollision->RenderComponent.Get();
-					if( HitLandscape )
+					if (HitLandscape)
 					{
 						TArray<uint8>* LayerCache = &LandscapeLayerCache->FindOrAdd(HitLandscape);
 						// TODO: FName to LayerInfo?
-						float HitWeight = HitLandscape->GetLayerWeightAtLocation( Instance.Location, HitLandscape->GetLandscapeInfo()->GetLayerInfoByName(LandscapeLayerName), LayerCache );
+						float HitWeight = HitLandscape->GetLayerWeightAtLocation(Instance.Location, HitLandscape->GetLandscapeInfo()->GetLayerInfoByName(LandscapeLayerName), LayerCache);
 
 						// Add count to bucket.
-						ExistingInstanceBuckets[FMath::RoundToInt(HitWeight * (float)(NUM_INSTANCE_BUCKETS-1))]++;
+						ExistingInstanceBuckets[FMath::RoundToInt(HitWeight * (float)(NUM_INSTANCE_BUCKETS - 1))]++;
 					}
 				}
 			}
@@ -615,7 +613,7 @@ void FEdModeFoliage::AddInstancesForBrush( UWorld* InWorld, AInstancedFoliageAct
 		else
 		{
 			// When not tied to a layer, put all the ExistingInstances in the last bucket.
-			ExistingInstanceBuckets[NUM_INSTANCE_BUCKETS-1] = ExistingInstances.Num();
+			ExistingInstanceBuckets[NUM_INSTANCE_BUCKETS - 1] = ExistingInstances.Num();
 		}
 
 		// We calculate a set of potential ExistingInstances for the brush area.
@@ -632,27 +630,27 @@ void FEdModeFoliage::AddInstancesForBrush( UWorld* InWorld, AInstancedFoliageAct
 		PotentialInstanceLocations.Empty(DesiredInstanceCount);
 
 		// Radius where we expect to have a single instance, given the density rules
-		const float DensityCheckRadius = FMath::Max<float>( FMath::Sqrt( (1000.f*1000.f) / (PI * MeshSettings->Density) ), MeshSettings->Radius );
+		const float DensityCheckRadius = FMath::Max<float>(FMath::Sqrt((1000.f*1000.f) / (PI * Settings->Density)), Settings->Radius);
 
-		for( int32 DesiredIdx=0;DesiredIdx<DesiredInstanceCount;DesiredIdx++ )
+		for (int32 DesiredIdx = 0; DesiredIdx < DesiredInstanceCount; DesiredIdx++)
 		{
 			FVector Start, End;
 
 			GetRandomVectorInBrush(Start, End);
-			
+
 			FHitResult Hit;
 			static FName NAME_AddInstancesForBrush = FName(TEXT("AddInstancesForBrush"));
 			FCollisionQueryParams TraceParams(NAME_AddInstancesForBrush, true, IFA);
-			if( MeshSettings->VertexColorMask != FOLIAGEVERTEXCOLORMASK_Disabled )
+			if (Settings->VertexColorMask != FOLIAGEVERTEXCOLORMASK_Disabled)
 			{
 				TraceParams.bReturnFaceIndex = true;
 			}
-			
-			if( InWorld->LineTraceSingle(Hit, Start, End, TraceParams, FCollisionObjectQueryParams(ECC_WorldStatic))) 
+
+			if (InWorld->LineTraceSingle(Hit, Start, End, TraceParams, FCollisionObjectQueryParams(ECC_WorldStatic)))
 			{
 				// Check filters
 				UPrimitiveComponent* PrimComp = Hit.Component.Get();
-				if( (PrimComp &&		 
+				if ((PrimComp &&
 					(PrimComp->GetOutermost() != InWorld->GetCurrentLevel()->GetOutermost() ||
 					(!UISettings.bFilterLandscape && PrimComp->IsA(ULandscapeHeightfieldCollisionComponent::StaticClass())) ||
 					(!UISettings.bFilterStaticMesh && PrimComp->IsA(UStaticMeshComponent::StaticClass())) ||
@@ -662,21 +660,21 @@ void FEdModeFoliage::AddInstancesForBrush( UWorld* InWorld, AInstancedFoliageAct
 					continue;
 				}
 
-				if( !CheckLocationForPotentialInstance( MeshInfo, MeshSettings, DensityCheckRadius, Hit.Location, Hit.Normal, PotentialInstanceLocations, PotentialInstanceHash ) )
+				if (!CheckLocationForPotentialInstance(MeshInfo, Settings, DensityCheckRadius, Hit.Location, Hit.Normal, PotentialInstanceLocations, PotentialInstanceHash))
 				{
 					continue;
 				}
 
 				// Check vertex color mask
-				if( MeshSettings->VertexColorMask != FOLIAGEVERTEXCOLORMASK_Disabled && Hit.FaceIndex != INDEX_NONE )
+				if (Settings->VertexColorMask != FOLIAGEVERTEXCOLORMASK_Disabled && Hit.FaceIndex != INDEX_NONE)
 				{
 					UStaticMeshComponent* HitStaticMeshComponent = Cast<UStaticMeshComponent>(Hit.Component.Get());
-					if( HitStaticMeshComponent )
+					if (HitStaticMeshComponent)
 					{
 						FColor VertexColor;
-						if( GetStaticMeshVertexColorForHit( HitStaticMeshComponent, Hit.FaceIndex, Hit.Location, VertexColor ) )
+						if (GetStaticMeshVertexColorForHit(HitStaticMeshComponent, Hit.FaceIndex, Hit.Location, VertexColor))
 						{
-							if( !CheckVertexColor(MeshSettings, VertexColor) )
+							if (!CheckVertexColor(Settings, VertexColor))
 							{
 								continue;
 							}
@@ -686,20 +684,20 @@ void FEdModeFoliage::AddInstancesForBrush( UWorld* InWorld, AInstancedFoliageAct
 
 				// Check landscape layer
 				float HitWeight = 1.f;
-				if( LandscapeLayerName != NAME_None )
+				if (LandscapeLayerName != NAME_None)
 				{
 					ULandscapeHeightfieldCollisionComponent* HitLandscapeCollision = Cast<ULandscapeHeightfieldCollisionComponent>(Hit.Component.Get());
-					if( HitLandscapeCollision )
+					if (HitLandscapeCollision)
 					{
 						ULandscapeComponent* HitLandscape = HitLandscapeCollision->RenderComponent.Get();
-						if( HitLandscape )
+						if (HitLandscape)
 						{
 							TArray<uint8>* LayerCache = &LandscapeLayerCache->FindOrAdd(HitLandscape);
 							// TODO: FName to LayerInfo?
-							HitWeight = HitLandscape->GetLayerWeightAtLocation( Hit.Location, HitLandscape->GetLandscapeInfo()->GetLayerInfoByName(LandscapeLayerName), LayerCache );
+							HitWeight = HitLandscape->GetLayerWeightAtLocation(Hit.Location, HitLandscape->GetLandscapeInfo()->GetLayerInfoByName(LandscapeLayerName), LayerCache);
 
 							// Reject instance randomly in proportion to weight
-							if( HitWeight <= FMath::FRand() )
+							if (HitWeight <= FMath::FRand())
 							{
 								continue;
 							}
@@ -707,27 +705,27 @@ void FEdModeFoliage::AddInstancesForBrush( UWorld* InWorld, AInstancedFoliageAct
 					}
 				}
 
-				const int32 BucketIndex = FMath::RoundToInt(HitWeight * (float)(NUM_INSTANCE_BUCKETS-1));
+				const int32 BucketIndex = FMath::RoundToInt(HitWeight * (float)(NUM_INSTANCE_BUCKETS - 1));
 				check(BucketIndex < ARRAY_COUNT(PotentialInstanceBuckets));
 
-				PotentialInstanceBuckets[BucketIndex].Add(FPotentialInstance( Hit.Location, Hit.Normal, Hit.Component.Get(), HitWeight ));
+				PotentialInstanceBuckets[BucketIndex].Add(FPotentialInstance(Hit.Location, Hit.Normal, Hit.Component.Get(), HitWeight));
 			}
 		}
 
-		for( int32 BucketIdx = 0; BucketIdx < NUM_INSTANCE_BUCKETS; BucketIdx++ )
+		for (int32 BucketIdx = 0; BucketIdx < NUM_INSTANCE_BUCKETS; BucketIdx++)
 		{
 			TArray<FPotentialInstance>& PotentialInstances = PotentialInstanceBuckets[BucketIdx];
-			float BucketFraction = (float)(BucketIdx+1) / (float)NUM_INSTANCE_BUCKETS;
+			float BucketFraction = (float)(BucketIdx + 1) / (float)NUM_INSTANCE_BUCKETS;
 
 			// We use the number that actually succeeded in placement (due to parameters) as the target
 			// for the number that should be in the brush region.
-			int32 AdditionalInstances = FMath::Clamp<int32>( FMath::RoundToInt( BucketFraction * (float)(PotentialInstances.Num() - ExistingInstanceBuckets[BucketIdx]) * Pressure), 0, PotentialInstances.Num() );
-			for( int32 Idx=0;Idx<AdditionalInstances;Idx++ )
+			int32 AdditionalInstances = FMath::Clamp<int32>(FMath::RoundToInt(BucketFraction * (float)(PotentialInstances.Num() - ExistingInstanceBuckets[BucketIdx]) * Pressure), 0, PotentialInstances.Num());
+			for (int32 Idx = 0; Idx < AdditionalInstances; Idx++)
 			{
 				FFoliageInstance Inst;
-				if (PotentialInstances[Idx].PlaceInstance(MeshSettings, Inst, InWorld ))
+				if (PotentialInstances[Idx].PlaceInstance(Settings, Inst, InWorld))
 				{
-					MeshInfo.AddInstance( IFA, StaticMesh, Inst );
+					MeshInfo.AddInstance(IFA, Settings, Inst);
 				}
 			}
 		}
@@ -735,28 +733,28 @@ void FEdModeFoliage::AddInstancesForBrush( UWorld* InWorld, AInstancedFoliageAct
 }
 
 /** Remove instances inside the brush to match DesiredInstanceCount */
-void FEdModeFoliage::RemoveInstancesForBrush( AInstancedFoliageActor* IFA, FFoliageMeshInfo& MeshInfo, int32 DesiredInstanceCount, TArray<int32>& ExistingInstances, float Pressure )
+void FEdModeFoliage::RemoveInstancesForBrush(AInstancedFoliageActor* IFA, FFoliageMeshInfo& MeshInfo, int32 DesiredInstanceCount, TArray<int32>& ExistingInstances, float Pressure)
 {
-	int32 InstancesToRemove = FMath::RoundToInt( (float)(ExistingInstances.Num() - DesiredInstanceCount) * Pressure);
+	int32 InstancesToRemove = FMath::RoundToInt((float)(ExistingInstances.Num() - DesiredInstanceCount) * Pressure);
 	int32 InstancesToKeep = ExistingInstances.Num() - InstancesToRemove;
-	if( InstancesToKeep > 0 )
-	{						
+	if (InstancesToKeep > 0)
+	{
 		// Remove InstancesToKeep random ExistingInstances from the array to leave those ExistingInstances behind, and delete all the rest
-		for( int32 i=0;i<InstancesToKeep;i++ )
+		for (int32 i = 0; i < InstancesToKeep; i++)
 		{
 			ExistingInstances.RemoveSwap(FMath::Rand() % ExistingInstances.Num());
 		}
 	}
 
-	if( !UISettings.bFilterLandscape || !UISettings.bFilterStaticMesh || !UISettings.bFilterBSP )
+	if (!UISettings.bFilterLandscape || !UISettings.bFilterStaticMesh || !UISettings.bFilterBSP)
 	{
 		// Filter ExistingInstances
-		for( int32 Idx=0;Idx<ExistingInstances.Num();Idx++ )
+		for (int32 Idx = 0; Idx < ExistingInstances.Num(); Idx++)
 		{
 			UActorComponent* Base = MeshInfo.Instances[ExistingInstances[Idx]].Base;
 
 			// Check if instance is candidate for removal based on filter settings
-			if( (Base && (
+			if ((Base && (
 				(!UISettings.bFilterLandscape && Base->IsA(ULandscapeHeightfieldCollisionComponent::StaticClass())) ||
 				(!UISettings.bFilterStaticMesh && Base->IsA(UStaticMeshComponent::StaticClass())) ||
 				(!UISettings.bFilterBSP && Base->IsA(UModelComponent::StaticClass()))
@@ -770,38 +768,40 @@ void FEdModeFoliage::RemoveInstancesForBrush( AInstancedFoliageActor* IFA, FFoli
 	}
 
 	// Remove ExistingInstances to reduce it to desired count
-	if( ExistingInstances.Num() > 0 )
+	if (ExistingInstances.Num() > 0)
 	{
-		MeshInfo.RemoveInstances( IFA, ExistingInstances );
+		MeshInfo.RemoveInstances(IFA, ExistingInstances);
 	}
 }
 
 /** Reapply instance settings to exiting instances */
-void FEdModeFoliage::ReapplyInstancesForBrush( UWorld* InWorld, AInstancedFoliageActor* IFA, FFoliageMeshInfo& MeshInfo, TArray<int32>& ExistingInstances )
+void FEdModeFoliage::ReapplyInstancesForBrush(UWorld* InWorld, AInstancedFoliageActor* IFA, UFoliageType* Settings, TArray<int32>& ExistingInstances)
 {
-	checkf( InWorld == IFA->GetWorld() , TEXT("Warning:World does not match Foliage world") );
-	UInstancedFoliageSettings* MeshSettings = MeshInfo.Settings;
-	bool bUpdated=false;
+	checkf(InWorld == IFA->GetWorld(), TEXT("Warning:World does not match Foliage world"));
+	FFoliageMeshInfo* MeshInfo = IFA->FindMesh(Settings);
+	check(MeshInfo);
+
+	bool bUpdated = false;
 
 	TArray<int32> UpdatedInstances;
 	TSet<int32> InstancesToDelete;
 
 	// Setup cache if we're reapplying landscape layer weights
-	FName LandscapeLayerName = MeshSettings->LandscapeLayer;
-	TMap<ULandscapeComponent*, TArray<uint8> >* LandscapeLayerCache = NULL;
-	if( MeshSettings->ReapplyLandscapeLayer && LandscapeLayerName != NAME_None )
+	FName LandscapeLayerName = Settings->LandscapeLayer;
+	TMap<ULandscapeComponent*, TArray<uint8>>* LandscapeLayerCache = nullptr;
+	if (Settings->ReapplyLandscapeLayer && LandscapeLayerName != NAME_None)
 	{
 		LandscapeLayerCache = &LandscapeLayerCaches.FindOrAdd(LandscapeLayerName);
 	}
 
 	IFA->Modify();
 
-	for( int32 Idx=0;Idx<ExistingInstances.Num();Idx++ )
+	for (int32 Idx = 0; Idx < ExistingInstances.Num(); Idx++)
 	{
 		int32 InstanceIndex = ExistingInstances[Idx];
-		FFoliageInstance& Instance = MeshInfo.Instances[InstanceIndex];
+		FFoliageInstance& Instance = MeshInfo->Instances[InstanceIndex];
 
-		if( (Instance.Flags & FOLIAGE_Readjusted) == 0 )
+		if ((Instance.Flags & FOLIAGE_Readjusted) == 0)
 		{
 			// record that we've made changes to this instance already, so we don't touch it again.
 			Instance.Flags |= FOLIAGE_Readjusted;
@@ -809,10 +809,10 @@ void FEdModeFoliage::ReapplyInstancesForBrush( UWorld* InWorld, AInstancedFoliag
 			bool bReapplyLocation = false;
 
 			// remove any Z offset first, so the offset is reapplied to any new 
-			if( FMath::Abs(Instance.ZOffset) > KINDA_SMALL_NUMBER )
+			if (FMath::Abs(Instance.ZOffset) > KINDA_SMALL_NUMBER)
 			{
-				MeshInfo.InstanceHash->RemoveInstance(Instance.Location, InstanceIndex);
-				Instance.Location = Instance.GetInstanceWorldTransform().TransformPosition(FVector(0,0,-Instance.ZOffset));
+				MeshInfo->InstanceHash->RemoveInstance(Instance.Location, InstanceIndex);
+				Instance.Location = Instance.GetInstanceWorldTransform().TransformPosition(FVector(0, 0, -Instance.ZOffset));
 				bReapplyLocation = true;
 			}
 
@@ -821,11 +821,11 @@ void FEdModeFoliage::ReapplyInstancesForBrush( UWorld* InWorld, AInstancedFoliag
 			bool bReapplyNormal = false;
 
 			// Reapply normal alignment
-			if( MeshSettings->ReapplyAlignToNormal )
+			if (Settings->ReapplyAlignToNormal)
 			{
-				if( MeshSettings->AlignToNormal )
+				if (Settings->AlignToNormal)
 				{
-					if( (Instance.Flags & FOLIAGE_AlignToNormal) == 0 )
+					if ((Instance.Flags & FOLIAGE_AlignToNormal) == 0)
 					{
 						bReapplyNormal = true;
 						bUpdated = true;
@@ -833,7 +833,7 @@ void FEdModeFoliage::ReapplyInstancesForBrush( UWorld* InWorld, AInstancedFoliag
 				}
 				else
 				{
-					if( Instance.Flags & FOLIAGE_AlignToNormal )
+					if (Instance.Flags & FOLIAGE_AlignToNormal)
 					{
 						Instance.Rotation = Instance.PreAlignRotation;
 						Instance.Flags &= (~FOLIAGE_AlignToNormal);
@@ -843,14 +843,14 @@ void FEdModeFoliage::ReapplyInstancesForBrush( UWorld* InWorld, AInstancedFoliag
 			}
 
 			// Reapply random yaw
-			if( MeshSettings->ReapplyRandomYaw )
+			if (Settings->ReapplyRandomYaw)
 			{
-				if( MeshSettings->RandomYaw )
+				if (Settings->RandomYaw)
 				{
-					if( Instance.Flags & FOLIAGE_NoRandomYaw )
+					if (Instance.Flags & FOLIAGE_NoRandomYaw)
 					{
 						// See if we need to remove any normal alignment first
-						if( !bReapplyNormal && (Instance.Flags & FOLIAGE_AlignToNormal) )
+						if (!bReapplyNormal && (Instance.Flags & FOLIAGE_AlignToNormal))
 						{
 							Instance.Rotation = Instance.PreAlignRotation;
 							bReapplyNormal = true;
@@ -862,10 +862,10 @@ void FEdModeFoliage::ReapplyInstancesForBrush( UWorld* InWorld, AInstancedFoliag
 				}
 				else
 				{
-					if( (Instance.Flags & FOLIAGE_NoRandomYaw)==0 )
+					if ((Instance.Flags & FOLIAGE_NoRandomYaw) == 0)
 					{
 						// See if we need to remove any normal alignment first
-						if( !bReapplyNormal && (Instance.Flags & FOLIAGE_AlignToNormal) )
+						if (!bReapplyNormal && (Instance.Flags & FOLIAGE_AlignToNormal))
 						{
 							Instance.Rotation = Instance.PreAlignRotation;
 							bReapplyNormal = true;
@@ -878,28 +878,28 @@ void FEdModeFoliage::ReapplyInstancesForBrush( UWorld* InWorld, AInstancedFoliag
 			}
 
 			// Reapply random pitch angle
-			if( MeshSettings->ReapplyRandomPitchAngle )
+			if (Settings->ReapplyRandomPitchAngle)
 			{
 				// See if we need to remove any normal alignment first
-				if( !bReapplyNormal && (Instance.Flags & FOLIAGE_AlignToNormal) )
+				if (!bReapplyNormal && (Instance.Flags & FOLIAGE_AlignToNormal))
 				{
 					Instance.Rotation = Instance.PreAlignRotation;
 					bReapplyNormal = true;
 				}
 
-				Instance.Rotation.Pitch = FMath::FRand() * MeshSettings->RandomPitchAngle;
+				Instance.Rotation.Pitch = FMath::FRand() * Settings->RandomPitchAngle;
 				Instance.Flags |= FOLIAGE_NoRandomYaw;
 
 				bUpdated = true;
 			}
 
 			// Reapply scale
-			if( MeshSettings->UniformScale )
+			if (Settings->UniformScale)
 			{
-				if( MeshSettings->ReapplyScaleX )
+				if (Settings->ReapplyScaleX)
 				{
-					float Scale = MeshSettings->ScaleMinX + FMath::FRand() * (MeshSettings->ScaleMaxX - MeshSettings->ScaleMinX);
-					Instance.DrawScale3D = FVector(Scale,Scale,Scale);
+					float Scale = Settings->ScaleMinX + FMath::FRand() * (Settings->ScaleMaxX - Settings->ScaleMinX);
+					Instance.DrawScale3D = FVector(Scale, Scale, Scale);
 					bUpdated = true;
 				}
 			}
@@ -907,52 +907,50 @@ void FEdModeFoliage::ReapplyInstancesForBrush( UWorld* InWorld, AInstancedFoliag
 			{
 				float LockRand;
 				// If we're doing axis scale locking, get an existing scale for a locked axis that we're not changing, for use as the locked scale value.
-				if( MeshSettings->LockScaleX && !MeshSettings->ReapplyScaleX && (MeshSettings->ScaleMaxX - MeshSettings->ScaleMinX) > KINDA_SMALL_NUMBER )
+				if (Settings->LockScaleX && !Settings->ReapplyScaleX && (Settings->ScaleMaxX - Settings->ScaleMinX) > KINDA_SMALL_NUMBER)
 				{
-					LockRand = (Instance.DrawScale3D.X - MeshSettings->ScaleMinX) / (MeshSettings->ScaleMaxX - MeshSettings->ScaleMinX);
+					LockRand = (Instance.DrawScale3D.X - Settings->ScaleMinX) / (Settings->ScaleMaxX - Settings->ScaleMinX);
 				}
-				else
-				if( MeshSettings->LockScaleY && !MeshSettings->ReapplyScaleY && (MeshSettings->ScaleMaxY - MeshSettings->ScaleMinY) > KINDA_SMALL_NUMBER )
+				else if (Settings->LockScaleY && !Settings->ReapplyScaleY && (Settings->ScaleMaxY - Settings->ScaleMinY) > KINDA_SMALL_NUMBER)
 				{
-					LockRand = (Instance.DrawScale3D.Y - MeshSettings->ScaleMinY) / (MeshSettings->ScaleMaxY - MeshSettings->ScaleMinY);
+					LockRand = (Instance.DrawScale3D.Y - Settings->ScaleMinY) / (Settings->ScaleMaxY - Settings->ScaleMinY);
 				}
-				else
-				if( MeshSettings->LockScaleZ && !MeshSettings->ReapplyScaleZ && (MeshSettings->ScaleMaxZ - MeshSettings->ScaleMinZ) > KINDA_SMALL_NUMBER )
+				else if (Settings->LockScaleZ && !Settings->ReapplyScaleZ && (Settings->ScaleMaxZ - Settings->ScaleMinZ) > KINDA_SMALL_NUMBER)
 				{
-					LockRand = (Instance.DrawScale3D.Z - MeshSettings->ScaleMinZ) / (MeshSettings->ScaleMaxZ - MeshSettings->ScaleMinZ);
+					LockRand = (Instance.DrawScale3D.Z - Settings->ScaleMinZ) / (Settings->ScaleMaxZ - Settings->ScaleMinZ);
 				}
 				else
 				{
 					LockRand = FMath::FRand();
 				}
 
-				if( MeshSettings->ReapplyScaleX )
+				if (Settings->ReapplyScaleX)
 				{
-					Instance.DrawScale3D.X = MeshSettings->ScaleMinX + (MeshSettings->LockScaleX ? LockRand : FMath::FRand()) * (MeshSettings->ScaleMaxX - MeshSettings->ScaleMinX);
+					Instance.DrawScale3D.X = Settings->ScaleMinX + (Settings->LockScaleX ? LockRand : FMath::FRand()) * (Settings->ScaleMaxX - Settings->ScaleMinX);
 					bUpdated = true;
 				}
 
-				if( MeshSettings->ReapplyScaleY )
+				if (Settings->ReapplyScaleY)
 				{
-					Instance.DrawScale3D.Y = MeshSettings->ScaleMinY + (MeshSettings->LockScaleY ? LockRand : FMath::FRand()) * (MeshSettings->ScaleMaxY - MeshSettings->ScaleMinY);
+					Instance.DrawScale3D.Y = Settings->ScaleMinY + (Settings->LockScaleY ? LockRand : FMath::FRand()) * (Settings->ScaleMaxY - Settings->ScaleMinY);
 					bUpdated = true;
 				}
 
-				if( MeshSettings->ReapplyScaleZ )
+				if (Settings->ReapplyScaleZ)
 				{
-					Instance.DrawScale3D.Z = MeshSettings->ScaleMinZ + (MeshSettings->LockScaleZ ? LockRand : FMath::FRand()) * (MeshSettings->ScaleMaxZ - MeshSettings->ScaleMinZ);
+					Instance.DrawScale3D.Z = Settings->ScaleMinZ + (Settings->LockScaleZ ? LockRand : FMath::FRand()) * (Settings->ScaleMaxZ - Settings->ScaleMinZ);
 					bUpdated = true;
 				}
 			}
 
-			if( MeshSettings->ReapplyZOffset )
+			if (Settings->ReapplyZOffset)
 			{
-				Instance.ZOffset = MeshSettings->ZOffsetMin + FMath::FRand() * (MeshSettings->ZOffsetMax - MeshSettings->ZOffsetMin);
+				Instance.ZOffset = Settings->ZOffsetMin + FMath::FRand() * (Settings->ZOffsetMax - Settings->ZOffsetMin);
 				bUpdated = true;
 			}
 
 			// Find a ground normal for either normal or ground slope check.
-			if( bReapplyNormal || MeshSettings->ReapplyGroundSlope || MeshSettings->ReapplyVertexColorMask || (MeshSettings->ReapplyLandscapeLayer && LandscapeLayerName != NAME_None) )
+			if (bReapplyNormal || Settings->ReapplyGroundSlope || Settings->ReapplyVertexColorMask || (Settings->ReapplyLandscapeLayer && LandscapeLayerName != NAME_None))
 			{
 				FHitResult Hit;
 				static const FName NAME_ReapplyInstancesForBrush = TEXT("ReapplyInstancesForBrush");
@@ -960,26 +958,26 @@ void FEdModeFoliage::ReapplyInstancesForBrush( UWorld* InWorld, AInstancedFoliag
 				// trace along the mesh's Z axis.
 				FVector ZAxis = Instance.Rotation.Quaternion().GetAxisZ();
 				FVector Start = Instance.Location + 16.f * ZAxis;
-				FVector End   = Instance.Location - 16.f * ZAxis;
+				FVector End = Instance.Location - 16.f * ZAxis;
 				FCollisionQueryParams TraceParams(NAME_ReapplyInstancesForBrush, true);
-				if( MeshSettings->ReapplyVertexColorMask && MeshSettings->VertexColorMask != FOLIAGEVERTEXCOLORMASK_Disabled )
+				if (Settings->ReapplyVertexColorMask && Settings->VertexColorMask != FOLIAGEVERTEXCOLORMASK_Disabled)
 				{
 					TraceParams.bReturnFaceIndex = true;
 				}
-				if( InWorld->LineTraceSingle(Hit, Start, End, TraceParams, FCollisionObjectQueryParams(ECC_WorldStatic))) 
+				if (InWorld->LineTraceSingle(Hit, Start, End, TraceParams, FCollisionObjectQueryParams(ECC_WorldStatic)))
 				{
 					// Reapply the normal
-					if( bReapplyNormal )
+					if (bReapplyNormal)
 					{
 						Instance.PreAlignRotation = Instance.Rotation;
-						Instance.AlignToNormal( Hit.Normal, MeshSettings->AlignMaxAngle );
+						Instance.AlignToNormal(Hit.Normal, Settings->AlignMaxAngle);
 					}
 
 					// Cull instances that don't meet the ground slope check.
-					if( MeshSettings->ReapplyGroundSlope )
+					if (Settings->ReapplyGroundSlope)
 					{
-						if( (MeshSettings->GroundSlope > 0.f && Hit.Normal.Z <= FMath::Cos (PI * MeshSettings->GroundSlope / 180.f)-SMALL_NUMBER) ||
-							(MeshSettings->GroundSlope < 0.f && Hit.Normal.Z >= FMath::Cos(-PI * MeshSettings->GroundSlope / 180.f)+SMALL_NUMBER) )
+						if ((Settings->GroundSlope > 0.f && Hit.Normal.Z <= FMath::Cos(PI * Settings->GroundSlope / 180.f) - SMALL_NUMBER) ||
+							(Settings->GroundSlope < 0.f && Hit.Normal.Z >= FMath::Cos(-PI * Settings->GroundSlope / 180.f) + SMALL_NUMBER))
 						{
 							InstancesToDelete.Add(InstanceIndex);
 							continue;
@@ -987,41 +985,41 @@ void FEdModeFoliage::ReapplyInstancesForBrush( UWorld* InWorld, AInstancedFoliag
 					}
 
 					// Cull instances for the landscape layer
-					if( MeshSettings->ReapplyLandscapeLayer )
+					if (Settings->ReapplyLandscapeLayer)
 					{
 						float HitWeight = 1.f;
 						ULandscapeHeightfieldCollisionComponent* HitLandscapeCollision = Cast<ULandscapeHeightfieldCollisionComponent>(Hit.Component.Get());
-						if( HitLandscapeCollision )
+						if (HitLandscapeCollision)
 						{
 							ULandscapeComponent* HitLandscape = HitLandscapeCollision->RenderComponent.Get();
-							if( HitLandscape )
+							if (HitLandscape)
 							{
 								TArray<uint8>* LayerCache = &LandscapeLayerCache->FindOrAdd(HitLandscape);
 								// TODO: FName to LayerInfo?
-								HitWeight = HitLandscape->GetLayerWeightAtLocation( Hit.Location, HitLandscape->GetLandscapeInfo()->GetLayerInfoByName(LandscapeLayerName), LayerCache );
+								HitWeight = HitLandscape->GetLayerWeightAtLocation(Hit.Location, HitLandscape->GetLandscapeInfo()->GetLayerInfoByName(LandscapeLayerName), LayerCache);
 
 								// Reject instance randomly in proportion to weight
-								if( HitWeight <= FMath::FRand() )
+								if (HitWeight <= FMath::FRand())
 								{
 									InstancesToDelete.Add(InstanceIndex);
 									continue;
 								}
-							}					
+							}
 						}
 					}
 
 					// Reapply vertex color mask
-					if( MeshSettings->ReapplyVertexColorMask )
+					if (Settings->ReapplyVertexColorMask)
 					{
-						if( MeshSettings->VertexColorMask != FOLIAGEVERTEXCOLORMASK_Disabled && Hit.FaceIndex != INDEX_NONE )
+						if (Settings->VertexColorMask != FOLIAGEVERTEXCOLORMASK_Disabled && Hit.FaceIndex != INDEX_NONE)
 						{
 							UStaticMeshComponent* HitStaticMeshComponent = Cast<UStaticMeshComponent>(Hit.Component.Get());
-							if( HitStaticMeshComponent )
+							if (HitStaticMeshComponent)
 							{
 								FColor VertexColor;
-								if( GetStaticMeshVertexColorForHit( HitStaticMeshComponent, Hit.FaceIndex, Hit.Location, VertexColor ) )
+								if (GetStaticMeshVertexColorForHit(HitStaticMeshComponent, Hit.FaceIndex, Hit.Location, VertexColor))
 								{
-									if( !CheckVertexColor(MeshSettings, VertexColor) )
+									if (!CheckVertexColor(Settings, VertexColor))
 									{
 										InstancesToDelete.Add(InstanceIndex);
 										continue;
@@ -1034,31 +1032,31 @@ void FEdModeFoliage::ReapplyInstancesForBrush( UWorld* InWorld, AInstancedFoliag
 			}
 
 			// Cull instances that don't meet the height range
-			if( MeshSettings->ReapplyHeight )
+			if (Settings->ReapplyHeight)
 			{
-				if( Instance.Location.Z < MeshSettings->HeightMin || Instance.Location.Z > MeshSettings->HeightMax )
+				if (Instance.Location.Z < Settings->HeightMin || Instance.Location.Z > Settings->HeightMax)
 				{
 					InstancesToDelete.Add(InstanceIndex);
 					continue;
 				}
 			}
 
-			if( bUpdated && FMath::Abs(Instance.ZOffset) > KINDA_SMALL_NUMBER )
+			if (bUpdated && FMath::Abs(Instance.ZOffset) > KINDA_SMALL_NUMBER)
 			{
 				// Reapply the Z offset in new local space
-				Instance.Location = Instance.GetInstanceWorldTransform().TransformPosition(FVector(0,0,Instance.ZOffset));
+				Instance.Location = Instance.GetInstanceWorldTransform().TransformPosition(FVector(0, 0, Instance.ZOffset));
 			}
 
 			// Re-add to the hash
-			if( bReapplyLocation )
+			if (bReapplyLocation)
 			{
-				MeshInfo.InstanceHash->InsertInstance(Instance.Location, InstanceIndex);
+				MeshInfo->InstanceHash->InsertInstance(Instance.Location, InstanceIndex);
 			}
 
 			// Cull overlapping based on radius
-			if( MeshSettings->ReapplyRadius && MeshSettings->Radius > 0.f )
+			if (Settings->ReapplyRadius && Settings->Radius > 0.f)
 			{
-				if( MeshInfo.CheckForOverlappingInstanceExcluding(InstanceIndex, MeshSettings->Radius, InstancesToDelete) )
+				if (MeshInfo->CheckForOverlappingInstanceExcluding(InstanceIndex, Settings->Radius, InstancesToDelete))
 				{
 					InstancesToDelete.Add(InstanceIndex);
 					continue;
@@ -1066,15 +1064,15 @@ void FEdModeFoliage::ReapplyInstancesForBrush( UWorld* InWorld, AInstancedFoliag
 			}
 
 			// Remove mesh collide with world
-			if( MeshSettings->ReapplyCollisionWithWorld )
+			if (Settings->ReapplyCollisionWithWorld)
 			{
 				FHitResult Hit;
 				static const FName NAME_ReapplyInstancesForBrush = TEXT("ReapplyCollisionWithWorld");
-				FVector Start = Instance.Location + FVector(0.f,0.f,16.f);
-				FVector End   = Instance.Location - FVector(0.f,0.f,16.f);
-				if( InWorld->LineTraceSingle(Hit, Start, End, FCollisionQueryParams(NAME_ReapplyInstancesForBrush, true), FCollisionObjectQueryParams(ECC_WorldStatic))) 
+				FVector Start = Instance.Location + FVector(0.f, 0.f, 16.f);
+				FVector End = Instance.Location - FVector(0.f, 0.f, 16.f);
+				if (InWorld->LineTraceSingle(Hit, Start, End, FCollisionQueryParams(NAME_ReapplyInstancesForBrush, true), FCollisionObjectQueryParams(ECC_WorldStatic)))
 				{
-					if (!CheckCollisionWithWorld(MeshSettings, Instance, Hit.Normal, Hit.Location, InWorld))
+					if (!CheckCollisionWithWorld(Settings, Instance, Hit.Normal, Hit.Location, InWorld))
 					{
 						InstancesToDelete.Add(InstanceIndex);
 						continue;
@@ -1086,33 +1084,33 @@ void FEdModeFoliage::ReapplyInstancesForBrush( UWorld* InWorld, AInstancedFoliag
 				}
 			}
 
-			if( bUpdated )
+			if (bUpdated)
 			{
 				UpdatedInstances.Add(InstanceIndex);
-			}		
+			}
 		}
 	}
 
-	if( UpdatedInstances.Num() > 0 )
+	if (UpdatedInstances.Num() > 0)
 	{
-		MeshInfo.PostUpdateInstances(IFA, UpdatedInstances);
+		MeshInfo->PostUpdateInstances(IFA, UpdatedInstances);
 		IFA->RegisterAllComponents();
 	}
-	
-	if( InstancesToDelete.Num() )
+
+	if (InstancesToDelete.Num())
 	{
-		MeshInfo.RemoveInstances(IFA, InstancesToDelete.Array());
+		MeshInfo->RemoveInstances(IFA, InstancesToDelete.Array());
 	}
 }
 
-void FEdModeFoliage::ApplyBrush( FEditorViewportClient* ViewportClient )
+void FEdModeFoliage::ApplyBrush(FEditorViewportClient* ViewportClient)
 {
-	if( !bBrushTraceValid )
+	if (!bBrushTraceValid)
 	{
 		return;
 	}
 
-	AInstancedFoliageActor* IFA = AInstancedFoliageActor::GetInstancedFoliageActor(GetWorld());
+	AInstancedFoliageActor* IFA = AInstancedFoliageActor::GetInstancedFoliageActorForCurrentLevel(GetWorld());
 
 	float BrushArea = PI * FMath::Square(UISettings.GetRadius());
 
@@ -1125,73 +1123,71 @@ void FEdModeFoliage::ApplyBrush( FEditorViewportClient* ViewportClient )
 	for (auto& MeshPair : IFA->FoliageMeshes)
 	{
 		FFoliageMeshInfo& MeshInfo = *MeshPair.Value;
-		UInstancedFoliageSettings* MeshSettings = MeshInfo.Settings;
+		UFoliageType* Settings = MeshPair.Key;
 
-		if( MeshSettings->IsSelected )
+		if (Settings->IsSelected)
 		{
 			// Find the instances already in the area.
 			TArray<int32> Instances;
-			FSphere BrushSphere(BrushLocation,UISettings.GetRadius());
-			MeshInfo.GetInstancesInsideSphere( BrushSphere, Instances );
+			FSphere BrushSphere(BrushLocation, UISettings.GetRadius());
+			MeshInfo.GetInstancesInsideSphere(BrushSphere, Instances);
 
-			if( UISettings.GetLassoSelectToolSelected() )
+			if (UISettings.GetLassoSelectToolSelected())
 			{
 				// Shift unpaints
-				MeshInfo.SelectInstances( IFA, !IsShiftDown(ViewportClient->Viewport), Instances);
+				MeshInfo.SelectInstances(IFA, !IsShiftDown(ViewportClient->Viewport), Instances);
 			}
 			else
-			if( UISettings.GetReapplyToolSelected() )
-			{
-				if( MeshSettings->ReapplyDensity )
+				if (UISettings.GetReapplyToolSelected())
+				{
+				if (Settings->ReapplyDensity)
 				{
 					// Adjust instance density
 					FMeshInfoSnapshot* SnapShot = InstanceSnapshot.Find(MeshPair.Key);
-					if( SnapShot ) 
+					if (SnapShot)
 					{
 						// Use snapshot to determine number of instances at the start of the brush stroke
-						int32 NewInstanceCount = FMath::RoundToInt( (float)SnapShot->CountInstancesInsideSphere(BrushSphere) * MeshSettings->ReapplyDensityAmount );
-						if( MeshSettings->ReapplyDensityAmount > 1.f && NewInstanceCount > Instances.Num() )
+						int32 NewInstanceCount = FMath::RoundToInt((float)SnapShot->CountInstancesInsideSphere(BrushSphere) * Settings->ReapplyDensityAmount);
+						if (Settings->ReapplyDensityAmount > 1.f && NewInstanceCount > Instances.Num())
 						{
-							AddInstancesForBrush( World , IFA, MeshPair.Key, MeshInfo, NewInstanceCount, Instances, Pressure );
+							AddInstancesForBrush(World, IFA, Settings, MeshInfo, NewInstanceCount, Instances, Pressure);
 						}
-						else
-						if( MeshSettings->ReapplyDensityAmount < 1.f && NewInstanceCount < Instances.Num() )			
+						else if (Settings->ReapplyDensityAmount < 1.f && NewInstanceCount < Instances.Num())
 						{
-							RemoveInstancesForBrush( IFA, MeshInfo, NewInstanceCount, Instances, Pressure );
+							RemoveInstancesForBrush(IFA, MeshInfo, NewInstanceCount, Instances, Pressure);
 						}
 					}
 				}
 
 				// Reapply any settings checked by the user
-				ReapplyInstancesForBrush( World , IFA, MeshInfo, Instances );
-			}
-			else
-			if( UISettings.GetPaintToolSelected() )
-			{	
-				// Shift unpaints
-				if( IsShiftDown(ViewportClient->Viewport) )
+				ReapplyInstancesForBrush(World, IFA, Settings, Instances);
+				}
+				else if (UISettings.GetPaintToolSelected())
 				{
-					int32 DesiredInstanceCount =  FMath::RoundToInt(BrushArea * MeshSettings->Density * UISettings.GetUnpaintDensity() / (1000.f*1000.f));
-
-					if( DesiredInstanceCount < Instances.Num() )
+					// Shift unpaints
+					if (IsShiftDown(ViewportClient->Viewport))
 					{
-						RemoveInstancesForBrush( IFA, MeshInfo, DesiredInstanceCount, Instances, Pressure );
+						int32 DesiredInstanceCount = FMath::RoundToInt(BrushArea * Settings->Density * UISettings.GetUnpaintDensity() / (1000.f*1000.f));
+
+						if (DesiredInstanceCount < Instances.Num())
+						{
+							RemoveInstancesForBrush(IFA, MeshInfo, DesiredInstanceCount, Instances, Pressure);
+						}
 					}
-				}
-				else
-				{
-					// This is the total set of instances disregarding parameters like slope, height or layer.
-					float DesiredInstanceCountFloat = BrushArea * MeshSettings->Density * UISettings.GetPaintDensity() / (1000.f*1000.f);
+					else
+					{
+						// This is the total set of instances disregarding parameters like slope, height or layer.
+						float DesiredInstanceCountFloat = BrushArea * Settings->Density * UISettings.GetPaintDensity() / (1000.f*1000.f);
 
-					// Allow a single instance with a random chance, if the brush is smaller than the density
-					int32 DesiredInstanceCount = DesiredInstanceCountFloat > 1.f ? FMath::RoundToInt(DesiredInstanceCountFloat) : FMath::FRand() < DesiredInstanceCountFloat ? 1 : 0;
+						// Allow a single instance with a random chance, if the brush is smaller than the density
+						int32 DesiredInstanceCount = DesiredInstanceCountFloat > 1.f ? FMath::RoundToInt(DesiredInstanceCountFloat) : FMath::FRand() < DesiredInstanceCountFloat ? 1 : 0;
 
-					AddInstancesForBrush( World, IFA, MeshPair.Key, MeshInfo, DesiredInstanceCount, Instances, Pressure );
+						AddInstancesForBrush(World, IFA, Settings, MeshInfo, DesiredInstanceCount, Instances, Pressure);
+					}
 				}
-			}
 		}
 	}
-	if( UISettings.GetLassoSelectToolSelected() )
+	if (UISettings.GetLassoSelectToolSelected())
 	{
 		IFA->CheckSelection();
 		Owner->PivotLocation = Owner->SnappedLocation = IFA->GetSelectionLocation();
@@ -1213,7 +1209,7 @@ struct FFoliagePaintBucketTriangle
 		WorldNormal = InLocalToWorld.GetDeterminant() >= 0.f ? Vector2 ^ Vector1 : Vector1 ^ Vector2;
 		float WorldNormalSize = WorldNormal.Size();
 		Area = WorldNormalSize * 0.5f;
-		if( WorldNormalSize > SMALL_NUMBER )
+		if (WorldNormalSize > SMALL_NUMBER)
 		{
 			WorldNormal /= WorldNormalSize;
 		}
@@ -1226,13 +1222,13 @@ struct FFoliagePaintBucketTriangle
 		float y = FMath::FRand();
 
 		// Flip if we're outside the triangle
-		if ( x + y > 1.f )
-		{ 
-			x = 1.f-x;
-			y = 1.f-y;
+		if (x + y > 1.f)
+		{
+			x = 1.f - x;
+			y = 1.f - y;
 		}
 
-		OutBaryVertexColor = (1.f-x-y) * VertexColor[0] + x * VertexColor[1] + y * VertexColor[2];
+		OutBaryVertexColor = (1.f - x - y) * VertexColor[0] + x * VertexColor[1] + y * VertexColor[2];
 		OutPoint = Vertex + x * Vector1 + y * Vector2;
 	}
 
@@ -1248,33 +1244,33 @@ struct FFoliagePaintBucketTriangle
 void FEdModeFoliage::ApplyPaintBucket(AActor* Actor, bool bRemove)
 {
 	// Apply only to current level
-	if( !Actor->GetLevel()->IsCurrentLevel() )
+	if (!Actor->GetLevel()->IsCurrentLevel())
 	{
 		return;
 	}
 
-	if( bRemove )
+	if (bRemove)
 	{
 		// Remove all instances of the selected meshes
-		AInstancedFoliageActor* IFA = AInstancedFoliageActor::GetInstancedFoliageActor(GetWorld());
+		AInstancedFoliageActor* IFA = AInstancedFoliageActor::GetInstancedFoliageActorForCurrentLevel(GetWorld());
 
 		TArray<UActorComponent*> Components;
 		Actor->GetComponents(Components);
 
-		for( int32 ComponentIdx=0;ComponentIdx<Components.Num();ComponentIdx++ )
+		for (int32 ComponentIdx = 0; ComponentIdx < Components.Num(); ComponentIdx++)
 		{
 			for (auto& MeshPair : IFA->FoliageMeshes)
 			{
 				FFoliageMeshInfo& MeshInfo = *MeshPair.Value;
-				UInstancedFoliageSettings* MeshSettings = MeshInfo.Settings;
+				UFoliageType* Settings = MeshPair.Key;
 
-				if( MeshSettings->IsSelected )
+				if (Settings->IsSelected)
 				{
 					FFoliageComponentHashInfo* ComponentHashInfo = MeshInfo.ComponentHash.Find(Components[ComponentIdx]);
-					if( ComponentHashInfo )
+					if (ComponentHashInfo)
 					{
 						TArray<int32> InstancesToRemove = ComponentHashInfo->Instances.Array();
-						MeshInfo.RemoveInstances( IFA, InstancesToRemove );
+						MeshInfo.RemoveInstances(IFA, InstancesToRemove);
 					}
 				}
 			}
@@ -1290,19 +1286,19 @@ void FEdModeFoliage::ApplyPaintBucket(AActor* Actor, bool bRemove)
 		TArray<UStaticMeshComponent*> StaticMeshComponents;
 		Actor->GetComponents(StaticMeshComponents);
 
-		for( int32 ComponentIdx=0;ComponentIdx<StaticMeshComponents.Num();ComponentIdx++ )
+		for (int32 ComponentIdx = 0; ComponentIdx < StaticMeshComponents.Num(); ComponentIdx++)
 		{
 			UStaticMeshComponent* StaticMeshComponent = StaticMeshComponents[ComponentIdx];
-			if( UISettings.bFilterStaticMesh && StaticMeshComponent->StaticMesh && StaticMeshComponent->StaticMesh->RenderData )
+			if (UISettings.bFilterStaticMesh && StaticMeshComponent->StaticMesh && StaticMeshComponent->StaticMesh->RenderData)
 			{
 				UStaticMesh* StaticMesh = StaticMeshComponent->StaticMesh;
 				FStaticMeshLODResources& LODModel = StaticMesh->RenderData->LODResources[0];
-				TArray<FFoliagePaintBucketTriangle>& PotentialTriangles = ComponentPotentialTriangles.Add(StaticMeshComponent, TArray<FFoliagePaintBucketTriangle>() );
+				TArray<FFoliagePaintBucketTriangle>& PotentialTriangles = ComponentPotentialTriangles.Add(StaticMeshComponent, TArray<FFoliagePaintBucketTriangle>());
 
 				bool bHasInstancedColorData = false;
-				const FStaticMeshComponentLODInfo* InstanceMeshLODInfo = NULL;
-				if( StaticMeshComponent->LODData.Num() > 0 )
-				{		
+				const FStaticMeshComponentLODInfo* InstanceMeshLODInfo = nullptr;
+				if (StaticMeshComponent->LODData.Num() > 0)
+				{
 					InstanceMeshLODInfo = StaticMeshComponent->LODData.GetTypedData();
 					bHasInstancedColorData = InstanceMeshLODInfo->PaintedVertices.Num() == LODModel.VertexBuffer.GetNumVertices();
 				}
@@ -1315,13 +1311,13 @@ void FEdModeFoliage::ApplyPaintBucket(AActor* Actor, bool bRemove)
 				const FColorVertexBuffer& ColorVertexBuffer = LODModel.ColorVertexBuffer;
 
 				// Build a mapping of vertex positions to vertex colors.  Using a TMap will allow for fast lookups so we can match new static mesh vertices with existing colors 
-				for( int32 Idx = 0; Idx < Indices.Num(); Idx+=3 )
+				for (int32 Idx = 0; Idx < Indices.Num(); Idx += 3)
 				{
 					const int32 Index0 = Indices[Idx];
-					const int32 Index1 = Indices[Idx+1];
-					const int32 Index2 = Indices[Idx+2];
+					const int32 Index1 = Indices[Idx + 1];
+					const int32 Index2 = Indices[Idx + 2];
 
-					new(PotentialTriangles) FFoliagePaintBucketTriangle(LocalToWorld
+					new(PotentialTriangles)FFoliagePaintBucketTriangle(LocalToWorld
 						, PositionVertexBuffer.VertexPosition(Index0)
 						, PositionVertexBuffer.VertexPosition(Index1)
 						, PositionVertexBuffer.VertexPosition(Index2)
@@ -1334,14 +1330,14 @@ void FEdModeFoliage::ApplyPaintBucket(AActor* Actor, bool bRemove)
 		}
 
 		// Place foliage
-		AInstancedFoliageActor* IFA = AInstancedFoliageActor::GetInstancedFoliageActor(GetWorld());
+		AInstancedFoliageActor* IFA = AInstancedFoliageActor::GetInstancedFoliageActorForCurrentLevel(GetWorld());
 
 		for (auto& MeshPair : IFA->FoliageMeshes)
 		{
 			FFoliageMeshInfo& MeshInfo = *MeshPair.Value;
-			UInstancedFoliageSettings* MeshSettings = MeshInfo.Settings;
+			UFoliageType* Settings = MeshPair.Key;
 
-			if( MeshSettings->IsSelected )
+			if (Settings->IsSelected)
 			{
 				// Quick lookup of potential instance locations, used for overlapping check.
 				TArray<FVector> PotentialInstanceLocations;
@@ -1349,44 +1345,44 @@ void FEdModeFoliage::ApplyPaintBucket(AActor* Actor, bool bRemove)
 				TArray<FPotentialInstance> InstancesToPlace;
 
 				// Radius where we expect to have a single instance, given the density rules
-				const float DensityCheckRadius = FMath::Max<float>( FMath::Sqrt( (1000.f*1000.f) / (PI * MeshSettings->Density * UISettings.GetPaintDensity() ) ), MeshSettings->Radius );
+				const float DensityCheckRadius = FMath::Max<float>(FMath::Sqrt((1000.f*1000.f) / (PI * Settings->Density * UISettings.GetPaintDensity())), Settings->Radius);
 
-				for( TMap<UPrimitiveComponent*, TArray<FFoliagePaintBucketTriangle> >::TIterator ComponentIt(ComponentPotentialTriangles); ComponentIt; ++ComponentIt )
+				for (TMap<UPrimitiveComponent*, TArray<FFoliagePaintBucketTriangle> >::TIterator ComponentIt(ComponentPotentialTriangles); ComponentIt; ++ComponentIt)
 				{
 					UPrimitiveComponent* Component = ComponentIt.Key();
 					TArray<FFoliagePaintBucketTriangle>& PotentialTriangles = ComponentIt.Value();
 
-					for( int32 TriIdx=0;TriIdx<PotentialTriangles.Num();TriIdx++ )
+					for (int32 TriIdx = 0; TriIdx<PotentialTriangles.Num(); TriIdx++)
 					{
 						FFoliagePaintBucketTriangle& Triangle = PotentialTriangles[TriIdx];
 
 						// Check if we can reject this triangle based on normal.
-						if( (MeshSettings->GroundSlope > 0.f && Triangle.WorldNormal.Z <= FMath::Cos( PI * MeshSettings->GroundSlope / 180.f)-SMALL_NUMBER) ||
-							(MeshSettings->GroundSlope < 0.f && Triangle.WorldNormal.Z >= FMath::Cos(-PI * MeshSettings->GroundSlope / 180.f)+SMALL_NUMBER) )
+						if ((Settings->GroundSlope > 0.f && Triangle.WorldNormal.Z <= FMath::Cos(PI * Settings->GroundSlope / 180.f) - SMALL_NUMBER) ||
+							(Settings->GroundSlope < 0.f && Triangle.WorldNormal.Z >= FMath::Cos(-PI * Settings->GroundSlope / 180.f) + SMALL_NUMBER))
 						{
 							continue;
 						}
 
 						// This is the total set of instances disregarding parameters like slope, height or layer.
-						float DesiredInstanceCountFloat = Triangle.Area * MeshSettings->Density * UISettings.GetPaintDensity() / (1000.f*1000.f);
+						float DesiredInstanceCountFloat = Triangle.Area * Settings->Density * UISettings.GetPaintDensity() / (1000.f*1000.f);
 
 						// Allow a single instance with a random chance, if the brush is smaller than the density
 						int32 DesiredInstanceCount = DesiredInstanceCountFloat > 1.f ? FMath::RoundToInt(DesiredInstanceCountFloat) : FMath::FRand() < DesiredInstanceCountFloat ? 1 : 0;
-				
-						for( int32 Idx=0;Idx<DesiredInstanceCount;Idx++ )
+
+						for (int32 Idx = 0; Idx < DesiredInstanceCount; Idx++)
 						{
 							FVector InstLocation;
 							FColor VertexColor;
 							Triangle.GetRandomPoint(InstLocation, VertexColor);
 
 							// Check color mask and filters at this location
-							if( !CheckVertexColor(MeshSettings, VertexColor) ||
-								!CheckLocationForPotentialInstance( MeshInfo, MeshSettings, DensityCheckRadius, InstLocation, Triangle.WorldNormal, PotentialInstanceLocations, PotentialInstanceHash ) )
+							if (!CheckVertexColor(Settings, VertexColor) ||
+								!CheckLocationForPotentialInstance(MeshInfo, Settings, DensityCheckRadius, InstLocation, Triangle.WorldNormal, PotentialInstanceLocations, PotentialInstanceHash))
 							{
 								continue;
 							}
-													
-							new(InstancesToPlace) FPotentialInstance( InstLocation, Triangle.WorldNormal, Component, 1.f );
+
+							new(InstancesToPlace)FPotentialInstance(InstLocation, Triangle.WorldNormal, Component, 1.f);
 						}
 					}
 				}
@@ -1395,12 +1391,12 @@ void FEdModeFoliage::ApplyPaintBucket(AActor* Actor, bool bRemove)
 				UWorld* World = Actor->GetWorld();
 				check(World);
 				// Place instances
-				for( int32 Idx=0;Idx<InstancesToPlace.Num();Idx++ )
+				for (int32 Idx = 0; Idx < InstancesToPlace.Num(); Idx++)
 				{
 					FFoliageInstance Inst;
-					if (InstancesToPlace[Idx].PlaceInstance(MeshSettings, Inst, World ))
+					if (InstancesToPlace[Idx].PlaceInstance(Settings, Inst, World))
 					{
-						MeshInfo.AddInstance( IFA, MeshPair.Key, Inst );
+						MeshInfo.AddInstance(IFA, Settings, Inst);
 					}
 				}
 			}
@@ -1408,20 +1404,20 @@ void FEdModeFoliage::ApplyPaintBucket(AActor* Actor, bool bRemove)
 	}
 }
 
-bool FEdModeFoliage::GetStaticMeshVertexColorForHit( UStaticMeshComponent* InStaticMeshComponent, int32 InTriangleIndex, const FVector& InHitLocation, FColor& OutVertexColor ) const
+bool FEdModeFoliage::GetStaticMeshVertexColorForHit(UStaticMeshComponent* InStaticMeshComponent, int32 InTriangleIndex, const FVector& InHitLocation, FColor& OutVertexColor) const
 {
 	UStaticMesh* StaticMesh = InStaticMeshComponent->StaticMesh;
 
-	if( StaticMesh == NULL || StaticMesh->RenderData == NULL )
+	if (StaticMesh == nullptr || StaticMesh->RenderData == nullptr)
 	{
 		return false;
 	}
 
 	FStaticMeshLODResources& LODModel = StaticMesh->RenderData->LODResources[0];
 	bool bHasInstancedColorData = false;
-	const FStaticMeshComponentLODInfo* InstanceMeshLODInfo = NULL;
-	if( InStaticMeshComponent->LODData.Num() > 0 )
-	{		
+	const FStaticMeshComponentLODInfo* InstanceMeshLODInfo = nullptr;
+	if (InStaticMeshComponent->LODData.Num() > 0)
+	{
 		InstanceMeshLODInfo = InStaticMeshComponent->LODData.GetTypedData();
 		bHasInstancedColorData = InstanceMeshLODInfo->PaintedVertices.Num() == LODModel.VertexBuffer.GetNumVertices();
 	}
@@ -1429,7 +1425,7 @@ bool FEdModeFoliage::GetStaticMeshVertexColorForHit( UStaticMeshComponent* InSta
 	const FColorVertexBuffer& ColorVertexBuffer = LODModel.ColorVertexBuffer;
 
 	// no vertex color data
-	if( !bHasInstancedColorData && ColorVertexBuffer.GetNumVertices() == 0 )
+	if (!bHasInstancedColorData && ColorVertexBuffer.GetNumVertices() == 0)
 	{
 		return false;
 	}
@@ -1439,22 +1435,22 @@ bool FEdModeFoliage::GetStaticMeshVertexColorForHit( UStaticMeshComponent* InSta
 	const FPositionVertexBuffer& PositionVertexBuffer = LODModel.PositionVertexBuffer;
 
 	int32 SectionFirstTriIndex = 0;
-	for(TArray<FStaticMeshSection>::TConstIterator SectionIt(LODModel.Sections);SectionIt;++SectionIt)
+	for (TArray<FStaticMeshSection>::TConstIterator SectionIt(LODModel.Sections); SectionIt; ++SectionIt)
 	{
 		const FStaticMeshSection& Section = *SectionIt;
 
 		if (Section.bEnableCollision)
 		{
 			int32 NextSectionTriIndex = SectionFirstTriIndex + Section.NumTriangles;
-			if( InTriangleIndex >= SectionFirstTriIndex && InTriangleIndex < NextSectionTriIndex )
+			if (InTriangleIndex >= SectionFirstTriIndex && InTriangleIndex < NextSectionTriIndex)
 			{
 
 				int32 IndexBufferIdx = (InTriangleIndex - SectionFirstTriIndex) * 3 + Section.FirstIndex;
 
 				// Look up the triangle vertex indices
 				int32 Index0 = Indices[IndexBufferIdx];
-				int32 Index1 = Indices[IndexBufferIdx+1];
-				int32 Index2 = Indices[IndexBufferIdx+2];
+				int32 Index1 = Indices[IndexBufferIdx + 1];
+				int32 Index2 = Indices[IndexBufferIdx + 2];
 
 				// Lookup the triangle world positions and colors.
 				FVector WorldVert0 = InStaticMeshComponent->ComponentToWorld.TransformPosition(PositionVertexBuffer.VertexPosition(Index0));
@@ -1465,7 +1461,7 @@ bool FEdModeFoliage::GetStaticMeshVertexColorForHit( UStaticMeshComponent* InSta
 				FLinearColor Color1;
 				FLinearColor Color2;
 
-				if( bHasInstancedColorData )
+				if (bHasInstancedColorData)
 				{
 					Color0 = InstanceMeshLODInfo->PaintedVertices[Index0].Color.ReinterpretAsLinear();
 					Color1 = InstanceMeshLODInfo->PaintedVertices[Index1].Color.ReinterpretAsLinear();
@@ -1492,7 +1488,7 @@ bool FEdModeFoliage::GetStaticMeshVertexColorForHit( UStaticMeshComponent* InSta
 			SectionFirstTriIndex = NextSectionTriIndex;
 		}
 	}
-	
+
 	return false;
 }
 
@@ -1502,7 +1498,7 @@ bool FEdModeFoliage::GetStaticMeshVertexColorForHit( UStaticMeshComponent* InSta
 
 TArray<struct FFoliageMeshUIInfo>& FEdModeFoliage::GetFoliageMeshList()
 {
-	UpdateFoliageMeshList(); 
+	UpdateFoliageMeshList();
 	return FoliageMeshList;
 }
 
@@ -1510,30 +1506,40 @@ void FEdModeFoliage::UpdateFoliageMeshList()
 {
 	FoliageMeshList.Empty();
 
-	AInstancedFoliageActor* IFA = AInstancedFoliageActor::GetInstancedFoliageActor(GetWorld());
+	AInstancedFoliageActor* IFA = AInstancedFoliageActor::GetInstancedFoliageActorForCurrentLevel(GetWorld());
 
 	for (auto& MeshPair : IFA->FoliageMeshes)
 	{
-		new(FoliageMeshList) FFoliageMeshUIInfo(MeshPair.Key, &*MeshPair.Value);
+		new(FoliageMeshList)FFoliageMeshUIInfo(MeshPair.Key, &*MeshPair.Value);
 	}
 
-	struct FCompareFFoliageMeshUIInfo
+	auto CompareDisplayOrder = [](const FFoliageMeshUIInfo& A, const FFoliageMeshUIInfo& B)
 	{
-		FORCEINLINE bool operator()( const FFoliageMeshUIInfo& A, const FFoliageMeshUIInfo& B ) const
-		{
-			return A.MeshInfo->Settings->DisplayOrder < B.MeshInfo->Settings->DisplayOrder;
-		}
+		return A.Settings->DisplayOrder < B.Settings->DisplayOrder;
 	};
-	FoliageMeshList.Sort( FCompareFFoliageMeshUIInfo() );
+	FoliageMeshList.Sort(CompareDisplayOrder);
 }
 
-/** Add a new mesh */
-void FEdModeFoliage::AddFoliageMesh(UStaticMesh* StaticMesh)
+/** Add a new mesh (if it doesn't already exist) */
+void FEdModeFoliage::AddFoliageMesh(UStaticMesh* Mesh)
 {
-	AInstancedFoliageActor* IFA = AInstancedFoliageActor::GetInstancedFoliageActor(GetWorld());
-	if( IFA->FindMesh(StaticMesh) == NULL )
+	AInstancedFoliageActor* IFA = AInstancedFoliageActor::GetInstancedFoliageActorForCurrentLevel(GetWorld());
+
+	if (IFA->GetSettingsForMesh(Mesh) == nullptr)
 	{
-		IFA->AddMesh(StaticMesh);
+		IFA->AddMesh(Mesh);
+
+		// Update mesh list.
+		UpdateFoliageMeshList();
+	}
+}
+
+void FEdModeFoliage::AddFoliageMesh(UFoliageType* Settings)
+{
+	AInstancedFoliageActor* IFA = AInstancedFoliageActor::GetInstancedFoliageActorForCurrentLevel(GetWorld());
+	if (IFA->FindMesh(Settings) == nullptr)
+	{
+		IFA->AddMesh(Settings);
 
 		// Update mesh list.
 		UpdateFoliageMeshList();
@@ -1541,21 +1547,21 @@ void FEdModeFoliage::AddFoliageMesh(UStaticMesh* StaticMesh)
 }
 
 /** Remove a mesh */
-bool FEdModeFoliage::RemoveFoliageMesh(UStaticMesh* StaticMesh)
+bool FEdModeFoliage::RemoveFoliageMesh(UFoliageType* Settings)
 {
 	bool bMeshRemoved = false;
 
-	AInstancedFoliageActor* IFA = AInstancedFoliageActor::GetInstancedFoliageActor(GetWorld());
-	FFoliageMeshInfo* MeshInfo = IFA->FindMesh(StaticMesh);
-	if( MeshInfo != NULL )
+	AInstancedFoliageActor* IFA = AInstancedFoliageActor::GetInstancedFoliageActorForCurrentLevel(GetWorld());
+	FFoliageMeshInfo* MeshInfo = IFA->FindMesh(Settings);
+	if (MeshInfo != nullptr)
 	{
 		int32 InstancesNum = MeshInfo->Instances.Num() - MeshInfo->FreeInstanceIndices.Num();
-		FText Message = FText::Format( NSLOCTEXT("UnrealEd", "FoliageMode_DeleteMesh", "Are you sure you want to remove all {0} instances of {1} from this level?"), FText::AsNumber( InstancesNum ), FText::FromString( StaticMesh->GetName() ) );
-		
-		if( InstancesNum == 0 || EAppReturnType::Yes == FMessageDialog::Open( EAppMsgType::YesNo, Message ) )
+		FText Message = FText::Format(NSLOCTEXT("UnrealEd", "FoliageMode_DeleteMesh", "Are you sure you want to remove all {0} instances of {1} from this level?"), FText::AsNumber(InstancesNum), FText::FromName(Settings->GetStaticMesh()->GetFName()));
+
+		if (InstancesNum == 0 || EAppReturnType::Yes == FMessageDialog::Open(EAppMsgType::YesNo, Message))
 		{
-			GEditor->BeginTransaction( NSLOCTEXT("UnrealEd", "FoliageMode_RemoveMeshTransaction", "Foliage Editing: Remove Mesh") );
-			IFA->RemoveMesh(StaticMesh);
+			GEditor->BeginTransaction(NSLOCTEXT("UnrealEd", "FoliageMode_RemoveMeshTransaction", "Foliage Editing: Remove Mesh"));
+			IFA->RemoveMesh(Settings);
 			GEditor->EndTransaction();
 
 			bMeshRemoved = true;
@@ -1569,35 +1575,35 @@ bool FEdModeFoliage::RemoveFoliageMesh(UStaticMesh* StaticMesh)
 }
 
 /** Bake instances to StaticMeshActors */
-void FEdModeFoliage::BakeFoliage(UStaticMesh* StaticMesh, bool bSelectedOnly)
+void FEdModeFoliage::BakeFoliage(UFoliageType* Settings, bool bSelectedOnly)
 {
-	AInstancedFoliageActor* IFA = AInstancedFoliageActor::GetInstancedFoliageActor(GetWorld());
-	FFoliageMeshInfo* MeshInfo = IFA->FindMesh(StaticMesh);
-	if( MeshInfo != NULL )
+	AInstancedFoliageActor* IFA = AInstancedFoliageActor::GetInstancedFoliageActorForCurrentLevel(GetWorld());
+	FFoliageMeshInfo* MeshInfo = IFA->FindMesh(Settings);
+	if (MeshInfo != nullptr)
 	{
 		TArray<int32> InstancesToConvert;
-		if( bSelectedOnly )
+		if (bSelectedOnly)
 		{
-			InstancesToConvert = MeshInfo->SelectedIndices;	
+			InstancesToConvert = MeshInfo->SelectedIndices;
 		}
 		else
 		{
-			for( int32 InstanceIdx=0;InstanceIdx<MeshInfo->Instances.Num();InstanceIdx++ )
+			for (int32 InstanceIdx = 0; InstanceIdx < MeshInfo->Instances.Num(); InstanceIdx++)
 			{
 				InstancesToConvert.Add(InstanceIdx);
 			}
 		}
 
 		// Convert
-		for( int32 Idx=0;Idx<InstancesToConvert.Num();Idx++ )
+		for (int32 Idx = 0; Idx < InstancesToConvert.Num(); Idx++)
 		{
 			FFoliageInstance& Instance = MeshInfo->Instances[InstancesToConvert[Idx]];
 			// We need a world in which to spawn the actor. Use the one from the original instance.
 			UWorld* World = IFA->GetWorld();
-			check(World != NULL);
-			AStaticMeshActor* SMA = World->SpawnActor<AStaticMeshActor>( Instance.Location, Instance.Rotation );
-			SMA->StaticMeshComponent->StaticMesh = StaticMesh;
-			SMA->GetRootComponent()->SetRelativeScale3D(Instance.DrawScale3D);	
+			check(World != nullptr);
+			AStaticMeshActor* SMA = World->SpawnActor<AStaticMeshActor>(Instance.Location, Instance.Rotation);
+			SMA->StaticMeshComponent->StaticMesh = Settings->GetStaticMesh();
+			SMA->GetRootComponent()->SetRelativeScale3D(Instance.DrawScale3D);
 			SMA->MarkComponentsRenderStateDirty();
 		}
 
@@ -1607,114 +1613,119 @@ void FEdModeFoliage::BakeFoliage(UStaticMesh* StaticMesh, bool bSelectedOnly)
 }
 
 /** Copy the settings object for this static mesh */
-UInstancedFoliageSettings* FEdModeFoliage::CopySettingsObject(UStaticMesh* StaticMesh)
+UFoliageType* FEdModeFoliage::CopySettingsObject(UFoliageType* Settings)
 {
-	AInstancedFoliageActor* IFA = AInstancedFoliageActor::GetInstancedFoliageActor(GetWorld());
-	FFoliageMeshInfo* MeshInfo = IFA->FindMesh(StaticMesh);
-	if( MeshInfo )
-	{
-		GEditor->BeginTransaction( NSLOCTEXT("UnrealEd", "FoliageMode_SettingsObjectTransaction", "Foliage Editing: Settings Object") );
-		IFA->Modify();
-		MeshInfo->Settings = Cast<UInstancedFoliageSettings>(StaticDuplicateObject(MeshInfo->Settings,IFA,NULL));
-		MeshInfo->Settings->ClearFlags(RF_Standalone|RF_Public);
-		GEditor->EndTransaction();
+	FScopedTransaction Transaction(NSLOCTEXT("UnrealEd", "FoliageMode_DuplicateSettingsObject", "Foliage Editing: Copy Settings Object"));
 
-		return MeshInfo->Settings;
+	AInstancedFoliageActor* IFA = AInstancedFoliageActor::GetInstancedFoliageActorForCurrentLevel(GetWorld());
+	IFA->Modify();
+
+	TUniqueObj<FFoliageMeshInfo> MeshInfo;
+	if (IFA->FoliageMeshes.RemoveAndCopyValue(Settings, MeshInfo))
+	{
+		Settings = (UFoliageType*)StaticDuplicateObject(Settings, IFA, nullptr, RF_AllFlags & ~(RF_Standalone | RF_Public));
+		IFA->FoliageMeshes.Add(Settings, MoveTemp(MeshInfo));
+		return Settings;
+	}
+	else
+	{
+		Transaction.Cancel();
 	}
 
-	return NULL;
+	return nullptr;
 }
 
 /** Replace the settings object for this static mesh with the one specified */
-void FEdModeFoliage::ReplaceSettingsObject(UStaticMesh* StaticMesh, UInstancedFoliageSettings* NewSettings)
+void FEdModeFoliage::ReplaceSettingsObject(UFoliageType* OldSettings, UFoliageType* NewSettings)
 {
-	AInstancedFoliageActor* IFA = AInstancedFoliageActor::GetInstancedFoliageActor(GetWorld());
-	FFoliageMeshInfo* MeshInfo = IFA->FindMesh(StaticMesh);
+	FScopedTransaction Transaction(NSLOCTEXT("UnrealEd", "FoliageMode_ReplaceSettingsObject", "Foliage Editing: Replace Settings Object"));
 
-	if( MeshInfo )
+	AInstancedFoliageActor* IFA = AInstancedFoliageActor::GetInstancedFoliageActorForCurrentLevel(GetWorld());
+	IFA->Modify();
+
+	TUniqueObj<FFoliageMeshInfo> MeshInfo;
+	if (IFA->FoliageMeshes.RemoveAndCopyValue(OldSettings, MeshInfo))
 	{
-		GEditor->BeginTransaction( NSLOCTEXT("UnrealEd", "FoliageMode_SettingsObjectTransaction", "Foliage Editing: Settings Object") );
-		IFA->Modify();
-		MeshInfo->Settings = NewSettings;
-		GEditor->EndTransaction();
+		IFA->FoliageMeshes.Add(NewSettings, MoveTemp(MeshInfo));
+	}
+	else
+	{
+		Transaction.Cancel();
 	}
 }
 
 /** Save the settings object */
-UInstancedFoliageSettings* FEdModeFoliage::SaveSettingsObject(const FText& InSettingsPackageName, UStaticMesh* StaticMesh)
+UFoliageType* FEdModeFoliage::SaveSettingsObject(const FText& InSettingsPackageName, UFoliageType* Settings)
 {
-	AInstancedFoliageActor* IFA = AInstancedFoliageActor::GetInstancedFoliageActor(GetWorld());
-	FFoliageMeshInfo* MeshInfo = IFA->FindMesh(StaticMesh);
+	FScopedTransaction Transaction(NSLOCTEXT("UnrealEd", "FoliageMode_SaveSettingsObject", "Foliage Editing: Save Settings Object"));
 
-	if( MeshInfo )
-	{
-		FString PackageName = InSettingsPackageName.ToString();
-		UPackage* Package = CreatePackage(NULL, *PackageName);
-		check(Package);
+	AInstancedFoliageActor* IFA = AInstancedFoliageActor::GetInstancedFoliageActorForCurrentLevel(GetWorld());
+	IFA->Modify();
 
-		GEditor->BeginTransaction( NSLOCTEXT("UnrealEd", "FoliageMode_SettingsObjectTransaction", "Foliage Editing: Settings Object") );
-		IFA->Modify();
-		MeshInfo->Settings = Cast<UInstancedFoliageSettings>(StaticDuplicateObject(MeshInfo->Settings,Package,*FPackageName::GetLongPackageAssetName(PackageName)));
-		MeshInfo->Settings->SetFlags(RF_Standalone|RF_Public);
-		GEditor->EndTransaction();
-			
-		// Notify the asset registry
-		FAssetRegistryModule::AssetCreated(MeshInfo->Settings);
+	FString PackageName = InSettingsPackageName.ToString();
+	UPackage* Package = CreatePackage(nullptr, *PackageName);
+	check(Package);
 
-		return MeshInfo->Settings;	
-	}
+	UFoliageType* NewSettings = Cast<UFoliageType>(StaticDuplicateObject(Settings, Package, *FPackageName::GetLongPackageAssetName(PackageName)));
+	NewSettings->SetFlags(RF_Standalone | RF_Public);
+
+	ReplaceSettingsObject(Settings, NewSettings);
 
-	return NULL;
+	// Notify the asset registry
+	FAssetRegistryModule::AssetCreated(NewSettings);
 
+	return NewSettings;
 }
 
 /** Reapply cluster settings to all the instances */
-void FEdModeFoliage::ReallocateClusters(UStaticMesh* StaticMesh)
+void FEdModeFoliage::ReallocateClusters(UFoliageType* Settings)
 {
-	AInstancedFoliageActor* IFA = AInstancedFoliageActor::GetInstancedFoliageActor(GetWorld());
-	FFoliageMeshInfo* MeshInfo = IFA->FindMesh(StaticMesh);
-	if( MeshInfo != NULL )
+	AInstancedFoliageActor* IFA = AInstancedFoliageActor::GetInstancedFoliageActorForCurrentLevel(GetWorld());
+	FFoliageMeshInfo* MeshInfo = IFA->FindMesh(Settings);
+	if (MeshInfo != nullptr)
 	{
-		MeshInfo->ReallocateClusters(IFA, StaticMesh);
+		MeshInfo->ReallocateClusters(IFA, Settings);
 	}
 }
 
 /** Replace a mesh with another one */
-bool FEdModeFoliage::ReplaceStaticMesh(UStaticMesh* OldStaticMesh, UStaticMesh* NewStaticMesh, bool& bOutMeshMerged)
+bool FEdModeFoliage::ReplaceStaticMesh(UFoliageType* OldSettings, UStaticMesh* NewStaticMesh, bool& bOutMeshMerged)
 {
 	bOutMeshMerged = false;
 
-	AInstancedFoliageActor* IFA = AInstancedFoliageActor::GetInstancedFoliageActor(GetWorld());
-	FFoliageMeshInfo* OldMeshInfo = IFA->FindMesh(OldStaticMesh);
+	AInstancedFoliageActor* IFA = AInstancedFoliageActor::GetInstancedFoliageActorForCurrentLevel(GetWorld());
+	FFoliageMeshInfo* OldMeshInfo = IFA->FindMesh(OldSettings);
 
-	if( OldMeshInfo != NULL && OldStaticMesh != NewStaticMesh )
+	if (OldMeshInfo != nullptr && OldSettings->GetStaticMesh() != NewStaticMesh)
 	{
 		int32 InstancesNum = OldMeshInfo->Instances.Num() - OldMeshInfo->FreeInstanceIndices.Num();
 
 		// Look for the new mesh in the mesh list, and either create a new mesh or merge the instances.
-		FFoliageMeshInfo* NewMeshInfo = IFA->FindMesh(NewStaticMesh);
-		if( NewMeshInfo == NULL )
+		UFoliageType* NewSettings = nullptr;
+		FFoliageMeshInfo* NewMeshInfo = nullptr;
+		NewSettings = IFA->GetSettingsForMesh(NewStaticMesh, &NewMeshInfo);
+
+		if (NewMeshInfo == nullptr)
 		{
-			FText Message = FText::Format( NSLOCTEXT("UnrealEd", "FoliageMode_ReplaceMesh", "Are you sure you want to replace all {0} instances of {1} with {2}?"), FText::AsNumber( InstancesNum ), FText::FromString( OldStaticMesh->GetName() ), FText::FromString( NewStaticMesh->GetName() ) );
+			FText Message = FText::Format(NSLOCTEXT("UnrealEd", "FoliageMode_ReplaceMesh", "Are you sure you want to replace all {0} instances of {1} with {2}?"), FText::AsNumber(InstancesNum), FText::FromString(OldSettings->GetStaticMesh()->GetName()), FText::FromString(NewStaticMesh->GetName()));
 
-			if( InstancesNum > 0 && EAppReturnType::Yes != FMessageDialog::Open(EAppMsgType::YesNo, Message ) )
+			if (InstancesNum > 0 && EAppReturnType::Yes != FMessageDialog::Open(EAppMsgType::YesNo, Message))
 			{
 				return false;
 			}
 
-			GEditor->BeginTransaction( NSLOCTEXT("UnrealEd", "FoliageMode_ChangeStaticMeshTransaction", "Foliage Editing: Change StaticMesh") );
+			GEditor->BeginTransaction(NSLOCTEXT("UnrealEd", "FoliageMode_ChangeStaticMeshTransaction", "Foliage Editing: Change StaticMesh"));
 			IFA->Modify();
-			NewMeshInfo = IFA->AddMesh(NewStaticMesh);
-			NewMeshInfo->Settings->DisplayOrder          = OldMeshInfo->Settings->DisplayOrder;
-			NewMeshInfo->Settings->ShowNothing          = OldMeshInfo->Settings->ShowNothing;
-			NewMeshInfo->Settings->ShowPaintSettings    = OldMeshInfo->Settings->ShowPaintSettings;
-			NewMeshInfo->Settings->ShowInstanceSettings = OldMeshInfo->Settings->ShowInstanceSettings;
+			NewMeshInfo = IFA->AddMesh(NewStaticMesh, &NewSettings);
+			NewSettings->DisplayOrder = OldSettings->DisplayOrder;
+			NewSettings->ShowNothing = OldSettings->ShowNothing;
+			NewSettings->ShowPaintSettings = OldSettings->ShowPaintSettings;
+			NewSettings->ShowInstanceSettings = OldSettings->ShowInstanceSettings;
 		}
-		else
-		if( InstancesNum > 0 &&
+		else if (InstancesNum > 0 &&
 			EAppReturnType::Yes != FMessageDialog::Open(EAppMsgType::YesNo,
-				FText::Format(NSLOCTEXT("UnrealEd", "FoliageMode_ReplaceMeshMerge", "Are you sure you want to merge all {0} instances of {1} into {2}?"),
-					FText::AsNumber(InstancesNum), FText::FromString(OldStaticMesh->GetName()), FText::FromString(NewStaticMesh->GetName()))) )
+			FText::Format(NSLOCTEXT("UnrealEd", "FoliageMode_ReplaceMeshMerge", "Are you sure you want to merge all {0} instances of {1} into {2}?"),
+			FText::AsNumber(InstancesNum), FText::FromString(OldSettings->GetStaticMesh()->GetName()), FText::FromString(NewStaticMesh->GetName()))))
 		{
 			return false;
 		}
@@ -1722,27 +1733,27 @@ bool FEdModeFoliage::ReplaceStaticMesh(UStaticMesh* OldStaticMesh, UStaticMesh*
 		{
 			bOutMeshMerged = true;
 
-			GEditor->BeginTransaction( NSLOCTEXT("UnrealEd", "FoliageMode_ChangeStaticMeshTransaction", "Foliage Editing: Change StaticMesh") );
+			GEditor->BeginTransaction(NSLOCTEXT("UnrealEd", "FoliageMode_ChangeStaticMeshTransaction", "Foliage Editing: Change StaticMesh"));
 			IFA->Modify();
 		}
 
-		if( InstancesNum > 0 )
+		if (InstancesNum > 0)
 		{
 			// copy instances from old to new.
-			for( int32 Idx=0;Idx<OldMeshInfo->Instances.Num();Idx++ )
+			for (FFoliageInstance& Instance : OldMeshInfo->Instances)
 			{
-				if( OldMeshInfo->Instances[Idx].ClusterIndex != -1 )
+				if (Instance.ClusterIndex != -1)
 				{
-					NewMeshInfo->AddInstance( IFA, NewStaticMesh, OldMeshInfo->Instances[Idx] );
+					NewMeshInfo->AddInstance(IFA, NewSettings, Instance);
 				}
 			}
 		}
 
 		// Remove the old mesh.
-		IFA->RemoveMesh(OldStaticMesh);
+		IFA->RemoveMesh(OldSettings);
 
 		GEditor->EndTransaction();
-		
+
 		// Update mesh list.
 		UpdateFoliageMeshList();
 	}
@@ -1751,32 +1762,32 @@ bool FEdModeFoliage::ReplaceStaticMesh(UStaticMesh* OldStaticMesh, UStaticMesh*
 
 
 /** FEdMode: Called when a key is pressed */
-bool FEdModeFoliage::InputKey( FEditorViewportClient* ViewportClient, FViewport* Viewport, FKey Key, EInputEvent Event )
+bool FEdModeFoliage::InputKey(FEditorViewportClient* ViewportClient, FViewport* Viewport, FKey Key, EInputEvent Event)
 {
-	if( UISettings.GetPaintToolSelected() || UISettings.GetReapplyToolSelected() || UISettings.GetLassoSelectToolSelected() )
+	if (UISettings.GetPaintToolSelected() || UISettings.GetReapplyToolSelected() || UISettings.GetLassoSelectToolSelected())
 	{
-		if( Key == EKeys::LeftMouseButton && Event == IE_Pressed && IsCtrlDown(Viewport) )
+		if (Key == EKeys::LeftMouseButton && Event == IE_Pressed && IsCtrlDown(Viewport))
 		{
-			if( !bToolActive )
+			if (!bToolActive)
 			{
-				GEditor->BeginTransaction( NSLOCTEXT("UnrealEd", "FoliageMode_EditTransaction", "Foliage Editing") );
+				GEditor->BeginTransaction(NSLOCTEXT("UnrealEd", "FoliageMode_EditTransaction", "Foliage Editing"));
 				InstanceSnapshot.Empty();
 
 				// Special setup beginning a stroke with the Reapply tool
 				// Necessary so we don't keep reapplying settings over and over for the same instances.
-				if( UISettings.GetReapplyToolSelected() )
+				if (UISettings.GetReapplyToolSelected())
 				{
-					AInstancedFoliageActor* IFA = AInstancedFoliageActor::GetInstancedFoliageActor(GetWorld());
+					AInstancedFoliageActor* IFA = AInstancedFoliageActor::GetInstancedFoliageActorForCurrentLevel(GetWorld());
 					for (auto& MeshPair : IFA->FoliageMeshes)
 					{
 						FFoliageMeshInfo& MeshInfo = *MeshPair.Value;
-						if( MeshInfo.Settings->IsSelected )
+						if (MeshPair.Key->IsSelected)
 						{
 							// Take a snapshot of all the locations
-							InstanceSnapshot.Add( MeshPair.Key, FMeshInfoSnapshot(&MeshInfo) ) ;
+							InstanceSnapshot.Add(MeshPair.Key, FMeshInfoSnapshot(&MeshInfo));
 
 							// Clear the "FOLIAGE_Readjusted" flag
-							for( int32 Idx=0;Idx<MeshInfo.Instances.Num();Idx++ )
+							for (int32 Idx = 0; Idx < MeshInfo.Instances.Num(); Idx++)
 							{
 								MeshInfo.Instances[Idx].Flags &= (~FOLIAGE_Readjusted);
 							}
@@ -1789,7 +1800,7 @@ bool FEdModeFoliage::InputKey( FEditorViewportClient* ViewportClient, FViewport*
 			return true;
 		}
 
-		if( Event == IE_Released && bToolActive && (Key == EKeys::LeftMouseButton || Key==EKeys::LeftControl || Key==EKeys::RightControl) )
+		if (Event == IE_Released && bToolActive && (Key == EKeys::LeftMouseButton || Key == EKeys::LeftControl || Key == EKeys::RightControl))
 		{
 			//Set the cursor position to that of the slate cursor so it wont snap back
 			Viewport->SetPreCaptureMousePosFromSlateCursor();
@@ -1801,19 +1812,19 @@ bool FEdModeFoliage::InputKey( FEditorViewportClient* ViewportClient, FViewport*
 		}
 	}
 
-	if( UISettings.GetSelectToolSelected() || UISettings.GetLassoSelectToolSelected() )
+	if (UISettings.GetSelectToolSelected() || UISettings.GetLassoSelectToolSelected())
 	{
-		if( Event == IE_Pressed )
+		if (Event == IE_Pressed)
 		{
-			if ( Key == EKeys::Platform_Delete )
+			if (Key == EKeys::Platform_Delete)
 			{
-				GEditor->BeginTransaction( NSLOCTEXT("UnrealEd", "FoliageMode_EditTransaction", "Foliage Editing") );
-				AInstancedFoliageActor* IFA = AInstancedFoliageActor::GetInstancedFoliageActor(GetWorld());
+				GEditor->BeginTransaction(NSLOCTEXT("UnrealEd", "FoliageMode_EditTransaction", "Foliage Editing"));
+				AInstancedFoliageActor* IFA = AInstancedFoliageActor::GetInstancedFoliageActorForCurrentLevel(GetWorld());
 				IFA->Modify();
 				for (auto& MeshPair : IFA->FoliageMeshes)
 				{
 					FFoliageMeshInfo& Mesh = *MeshPair.Value;
-					if( Mesh.SelectedIndices.Num() > 0 )
+					if (Mesh.SelectedIndices.Num() > 0)
 					{
 						TArray<int32> InstancesToDelete = Mesh.SelectedIndices;
 						Mesh.RemoveInstances(IFA, InstancesToDelete);
@@ -1823,35 +1834,32 @@ bool FEdModeFoliage::InputKey( FEditorViewportClient* ViewportClient, FViewport*
 
 				return true;
 			}
-			else
-			if( Key == EKeys::End )
+			else if (Key == EKeys::End)
 			{
 				// Snap instances to ground
-				AInstancedFoliageActor* IFA = AInstancedFoliageActor::GetInstancedFoliageActor(GetWorld());
+				AInstancedFoliageActor* IFA = AInstancedFoliageActor::GetInstancedFoliageActorForCurrentLevel(GetWorld());
 				IFA->Modify();
 				bool bMovedInstance = false;
 				for (auto& MeshPair : IFA->FoliageMeshes)
 				{
 					FFoliageMeshInfo& Mesh = *MeshPair.Value;
-					
+
 					Mesh.PreMoveInstances(IFA, Mesh.SelectedIndices);
 
 					UWorld* World = ViewportClient->GetWorld();
-					for( int32 SelectionIdx=0;SelectionIdx<Mesh.SelectedIndices.Num();SelectionIdx++ )
+					for (int32 SelectedInstanceIdx : Mesh.SelectedIndices)
 					{
-						int32 InstanceIdx = Mesh.SelectedIndices[SelectionIdx];
-
-						FFoliageInstance& Instance = Mesh.Instances[InstanceIdx];
+						FFoliageInstance& Instance = Mesh.Instances[SelectedInstanceIdx];
 
 						FVector Start = Instance.Location;
-						FVector End = Instance.Location - FVector(0.f,0.f,FOLIAGE_SNAP_TRACE);
+						FVector End = Instance.Location - FVector(0.f, 0.f, FOLIAGE_SNAP_TRACE);
 
 						FHitResult Hit;
-						if( World->LineTraceSingle(Hit, Start, End, FCollisionQueryParams(FName("FoliageSnap"), true), FCollisionObjectQueryParams(ECC_WorldStatic))) 
+						if (World->LineTraceSingle(Hit, Start, End, FCollisionQueryParams(FName("FoliageSnap"), true), FCollisionObjectQueryParams(ECC_WorldStatic)))
 						{
 							// Check current level
-							if( (Hit.Component.IsValid() && Hit.Component.Get()->GetOutermost() == World->GetCurrentLevel()->GetOutermost()) ||
-								(Hit.GetActor() && Hit.GetActor()->IsA(AWorldSettings::StaticClass())) )
+							if ((Hit.Component.IsValid() && Hit.Component.Get()->GetOutermost() == World->GetCurrentLevel()->GetOutermost()) ||
+								(Hit.GetActor() && Hit.GetActor()->IsA(AWorldSettings::StaticClass())))
 							{
 								Instance.Location = Hit.Location;
 								Instance.ZOffset = 0.f;
@@ -1859,14 +1867,14 @@ bool FEdModeFoliage::InputKey( FEditorViewportClient* ViewportClient, FViewport*
 								// We cannot be based on an a blueprint component as these will disappear when the construction script is re-run
 								if (Instance.Base && Instance.Base->bCreatedByConstructionScript)
 								{
-									Instance.Base = NULL;
+									Instance.Base = nullptr;
 								}
 
-								if( Instance.Flags & FOLIAGE_AlignToNormal )
+								if (Instance.Flags & FOLIAGE_AlignToNormal)
 								{
 									// Remove previous alignment and align to new normal.
 									Instance.Rotation = Instance.PreAlignRotation;
-									Instance.AlignToNormal( Hit.Normal, Mesh.Settings->AlignMaxAngle );
+									Instance.AlignToNormal(Hit.Normal, MeshPair.Key->AlignMaxAngle);
 								}
 							}
 						}
@@ -1876,7 +1884,7 @@ bool FEdModeFoliage::InputKey( FEditorViewportClient* ViewportClient, FViewport*
 					Mesh.PostMoveInstances(IFA, Mesh.SelectedIndices);
 				}
 
-				if( bMovedInstance )
+				if (bMovedInstance)
 				{
 					Owner->PivotLocation = Owner->SnappedLocation = IFA->GetSelectionLocation();
 					IFA->MarkComponentsRenderStateDirty();
@@ -1891,30 +1899,30 @@ bool FEdModeFoliage::InputKey( FEditorViewportClient* ViewportClient, FViewport*
 }
 
 /** FEdMode: Render the foliage edit mode */
-void FEdModeFoliage::Render( const FSceneView* View, FViewport* Viewport, FPrimitiveDrawInterface* PDI )
+void FEdModeFoliage::Render(const FSceneView* View, FViewport* Viewport, FPrimitiveDrawInterface* PDI)
 {
 	/** Call parent implementation */
-	FEdMode::Render( View, Viewport, PDI );
+	FEdMode::Render(View, Viewport, PDI);
 }
 
 
 /** FEdMode: Render HUD elements for this tool */
-void FEdModeFoliage::DrawHUD( FEditorViewportClient* ViewportClient, FViewport* Viewport, const FSceneView* View, FCanvas* Canvas )
+void FEdModeFoliage::DrawHUD(FEditorViewportClient* ViewportClient, FViewport* Viewport, const FSceneView* View, FCanvas* Canvas)
 {
 }
 
 /** FEdMode: Check to see if an actor can be selected in this mode - no side effects */
-bool FEdModeFoliage::IsSelectionAllowed( AActor* InActor, bool bInSelection ) const
+bool FEdModeFoliage::IsSelectionAllowed(AActor* InActor, bool bInSelection) const
 {
 	return false;
 }
 
 /** FEdMode: Handling SelectActor */
-bool FEdModeFoliage::Select( AActor* InActor, bool bInSelected )
+bool FEdModeFoliage::Select(AActor* InActor, bool bInSelected)
 {
 	// return true if you filter that selection
 	// however - return false if we are trying to deselect so that it will infact do the deselection
-	if( bInSelected == false )
+	if (bInSelected == false)
 	{
 		return false;
 	}
@@ -1928,67 +1936,67 @@ void FEdModeFoliage::ActorSelectionChangeNotify()
 
 
 /** Forces real-time perspective viewports */
-void FEdModeFoliage::ForceRealTimeViewports( const bool bEnable, const bool bStoreCurrentState )
+void FEdModeFoliage::ForceRealTimeViewports(const bool bEnable, const bool bStoreCurrentState)
 {
-	FLevelEditorModule& LevelEditorModule = FModuleManager::GetModuleChecked<FLevelEditorModule>( "LevelEditor");
+	FLevelEditorModule& LevelEditorModule = FModuleManager::GetModuleChecked<FLevelEditorModule>("LevelEditor");
 	TSharedPtr< ILevelViewport > ViewportWindow = LevelEditorModule.GetFirstActiveViewport();
 	if (ViewportWindow.IsValid())
 	{
 		FEditorViewportClient &Viewport = ViewportWindow->GetLevelViewportClient();
-		if( Viewport.IsPerspective() )
-		{				
-			if( bEnable )
+		if (Viewport.IsPerspective())
+		{
+			if (bEnable)
 			{
-				Viewport.SetRealtime( bEnable, bStoreCurrentState );
+				Viewport.SetRealtime(bEnable, bStoreCurrentState);
 			}
 			else
 			{
 				const bool bAllowDisable = true;
-				Viewport.RestoreRealtime( bAllowDisable );
+				Viewport.RestoreRealtime(bAllowDisable);
 			}
 
 		}
 	}
 }
 
-bool FEdModeFoliage::HandleClick(FEditorViewportClient* InViewportClient, HHitProxy *HitProxy, const FViewportClick &Click )
+bool FEdModeFoliage::HandleClick(FEditorViewportClient* InViewportClient, HHitProxy *HitProxy, const FViewportClick &Click)
 {
-	AInstancedFoliageActor* IFA = AInstancedFoliageActor::GetInstancedFoliageActor(GetWorld());
+	AInstancedFoliageActor* IFA = AInstancedFoliageActor::GetInstancedFoliageActorForCurrentLevel(GetWorld());
 
-	if( UISettings.GetSelectToolSelected() )
+	if (UISettings.GetSelectToolSelected())
 	{
-		if( HitProxy && HitProxy->IsA(HInstancedStaticMeshInstance::StaticGetType()) )
+		if (HitProxy && HitProxy->IsA(HInstancedStaticMeshInstance::StaticGetType()))
 		{
-			HInstancedStaticMeshInstance* SMIProxy = ( ( HInstancedStaticMeshInstance* )HitProxy );
-		
-			IFA->SelectInstance( SMIProxy->Component, SMIProxy->InstanceIndex, Click.IsControlDown() );
+			HInstancedStaticMeshInstance* SMIProxy = ((HInstancedStaticMeshInstance*)HitProxy);
+
+			IFA->SelectInstance(SMIProxy->Component, SMIProxy->InstanceIndex, Click.IsControlDown());
 
 			// Update pivot
 			Owner->PivotLocation = Owner->SnappedLocation = IFA->GetSelectionLocation();
 		}
 		else
 		{
-			if( !Click.IsControlDown() )
+			if (!Click.IsControlDown())
 			{
 				// Select none if not trying to toggle
-				IFA->SelectInstance( NULL, -1, false );
+				IFA->SelectInstance(nullptr, -1, false);
 			}
 		}
 
 		return true;
 	}
 	else
-	if( UISettings.GetPaintBucketToolSelected() || UISettings.GetReapplyPaintBucketToolSelected() )
-	{
-		if( HitProxy && HitProxy->IsA(HActor::StaticGetType()) && Click.IsControlDown() )
+		if (UISettings.GetPaintBucketToolSelected() || UISettings.GetReapplyPaintBucketToolSelected())
 		{
-			GEditor->BeginTransaction( NSLOCTEXT("UnrealEd", "FoliageMode_EditTransaction", "Foliage Editing") );
-			ApplyPaintBucket( ((HActor*)HitProxy)->Actor, Click.IsShiftDown() );
+		if (HitProxy && HitProxy->IsA(HActor::StaticGetType()) && Click.IsControlDown())
+		{
+			GEditor->BeginTransaction(NSLOCTEXT("UnrealEd", "FoliageMode_EditTransaction", "Foliage Editing"));
+			ApplyPaintBucket(((HActor*)HitProxy)->Actor, Click.IsShiftDown());
 			GEditor->EndTransaction();
 		}
 
 		return true;
-	}
+		}
 
 	return FEdMode::HandleClick(InViewportClient, HitProxy, Click);
 }
@@ -2001,16 +2009,16 @@ FVector FEdModeFoliage::GetWidgetLocation() const
 /** FEdMode: Called when a mouse button is pressed */
 bool FEdModeFoliage::StartTracking(FEditorViewportClient* InViewportClient, FViewport* InViewport)
 {
-	if( UISettings.GetSelectToolSelected() || UISettings.GetLassoSelectToolSelected() )
+	if (UISettings.GetSelectToolSelected() || UISettings.GetLassoSelectToolSelected())
 	{
 		// Update pivot
-		AInstancedFoliageActor* IFA = AInstancedFoliageActor::GetInstancedFoliageActor(GetWorld());
+		AInstancedFoliageActor* IFA = AInstancedFoliageActor::GetInstancedFoliageActorForCurrentLevel(GetWorld());
 		Owner->PivotLocation = Owner->SnappedLocation = IFA->GetSelectionLocation();
 
-		GEditor->BeginTransaction( NSLOCTEXT("UnrealEd", "FoliageMode_EditTransaction", "Foliage Editing") );
-		
+		GEditor->BeginTransaction(NSLOCTEXT("UnrealEd", "FoliageMode_EditTransaction", "Foliage Editing"));
+
 		bCanAltDrag = true;
-	
+
 		return true;
 	}
 	return FEdMode::StartTracking(InViewportClient, InViewport);
@@ -2019,7 +2027,7 @@ bool FEdModeFoliage::StartTracking(FEditorViewportClient* InViewportClient, FVie
 /** FEdMode: Called when the a mouse button is released */
 bool FEdModeFoliage::EndTracking(FEditorViewportClient* InViewportClient, FViewport* InViewport)
 {
-	if( UISettings.GetSelectToolSelected() || UISettings.GetLassoSelectToolSelected() )
+	if (UISettings.GetSelectToolSelected() || UISettings.GetLassoSelectToolSelected())
 	{
 		GEditor->EndTransaction();
 		return true;
@@ -2028,32 +2036,31 @@ bool FEdModeFoliage::EndTracking(FEditorViewportClient* InViewportClient, FViewp
 }
 
 /** FEdMode: Called when mouse drag input it applied */
-bool FEdModeFoliage::InputDelta( FEditorViewportClient* InViewportClient, FViewport* InViewport, FVector& InDrag, FRotator& InRot, FVector& InScale )
+bool FEdModeFoliage::InputDelta(FEditorViewportClient* InViewportClient, FViewport* InViewport, FVector& InDrag, FRotator& InRot, FVector& InScale)
 {
 	bool bFoundSelection = false;
-	
+
 	bool bAltDown = InViewport->KeyState(EKeys::LeftAlt) || InViewport->KeyState(EKeys::RightAlt);
 
-	if( InViewportClient->GetCurrentWidgetAxis() != EAxisList::None && (UISettings.GetSelectToolSelected() || UISettings.GetLassoSelectToolSelected()) )
+	if (InViewportClient->GetCurrentWidgetAxis() != EAxisList::None && (UISettings.GetSelectToolSelected() || UISettings.GetLassoSelectToolSelected()))
 	{
-		AInstancedFoliageActor* IFA = AInstancedFoliageActor::GetInstancedFoliageActor(GetWorld());
+		AInstancedFoliageActor* IFA = AInstancedFoliageActor::GetInstancedFoliageActorForCurrentLevel(GetWorld());
 		IFA->Modify();
 		for (auto& MeshPair : IFA->FoliageMeshes)
 		{
 			FFoliageMeshInfo& MeshInfo = *MeshPair.Value;
 			bFoundSelection |= MeshInfo.SelectedIndices.Num() > 0;
 
-			if( bAltDown && bCanAltDrag && (InViewportClient->GetCurrentWidgetAxis() & EAxisList::XYZ)  )
+			if (bAltDown && bCanAltDrag && (InViewportClient->GetCurrentWidgetAxis() & EAxisList::XYZ))
 			{
 				MeshInfo.DuplicateInstances(IFA, MeshPair.Key, MeshInfo.SelectedIndices);
 			}
-			
+
 			MeshInfo.PreMoveInstances(IFA, MeshInfo.SelectedIndices);
 
-			for( int32 SelectionIdx=0;SelectionIdx<MeshInfo.SelectedIndices.Num();SelectionIdx++ )
+			for (int32 SelectedInstanceIdx : MeshInfo.SelectedIndices)
 			{
-				int32 InstanceIndex = MeshInfo.SelectedIndices[SelectionIdx];
-				FFoliageInstance& Instance = MeshInfo.Instances[InstanceIndex];
+				FFoliageInstance& Instance = MeshInfo.Instances[SelectedInstanceIdx];
 				Instance.Location += InDrag;
 				Instance.ZOffset = 0.f;
 				Instance.Rotation += InRot;
@@ -2066,14 +2073,14 @@ bool FEdModeFoliage::InputDelta( FEditorViewportClient* InViewportClient, FViewp
 		// Only allow alt-drag on first InputDelta
 		bCanAltDrag = false;
 
-		if( bFoundSelection )
+		if (bFoundSelection)
 		{
 			IFA->MarkComponentsRenderStateDirty();
 			return true;
 		}
 	}
 
-	return FEdMode::InputDelta(InViewportClient,InViewport,InDrag,InRot,InScale);
+	return FEdMode::InputDelta(InViewportClient, InViewport, InDrag, InRot, InScale);
 }
 
 bool FEdModeFoliage::AllowWidgetMove()
@@ -2088,12 +2095,12 @@ bool FEdModeFoliage::UsesTransformWidget() const
 
 bool FEdModeFoliage::ShouldDrawWidget() const
 {
-	return (UISettings.GetSelectToolSelected() || (UISettings.GetLassoSelectToolSelected() && !bToolActive)) && AInstancedFoliageActor::GetInstancedFoliageActor(GetWorld())->SelectedMesh != NULL;
+	return (UISettings.GetSelectToolSelected() || (UISettings.GetLassoSelectToolSelected() && !bToolActive)) && AInstancedFoliageActor::GetInstancedFoliageActorForCurrentLevel(GetWorld())->SelectedMesh != nullptr;
 }
 
-EAxisList::Type FEdModeFoliage::GetWidgetAxisToDraw( FWidget::EWidgetMode InWidgetMode ) const
+EAxisList::Type FEdModeFoliage::GetWidgetAxisToDraw(FWidget::EWidgetMode InWidgetMode) const
 {
-	switch(InWidgetMode)
+	switch (InWidgetMode)
 	{
 	case FWidget::WM_Translate:
 	case FWidget::WM_Rotate:
@@ -2108,36 +2115,36 @@ EAxisList::Type FEdModeFoliage::GetWidgetAxisToDraw( FWidget::EWidgetMode InWidg
 void FFoliageUISettings::Load()
 {
 	FString WindowPositionString;
-	if( GConfig->GetString( TEXT("FoliageEdit"), TEXT("WindowPosition"), WindowPositionString, GEditorUserSettingsIni ) )
+	if (GConfig->GetString(TEXT("FoliageEdit"), TEXT("WindowPosition"), WindowPositionString, GEditorUserSettingsIni))
 	{
 		TArray<FString> PositionValues;
-		if( WindowPositionString.ParseIntoArray( &PositionValues, TEXT( "," ), true ) == 4 )
+		if (WindowPositionString.ParseIntoArray(&PositionValues, TEXT(","), true) == 4)
 		{
-			WindowX = FCString::Atoi( *PositionValues[0] );
-			WindowY = FCString::Atoi( *PositionValues[1] );
-			WindowWidth = FCString::Atoi( *PositionValues[2] );
-			WindowHeight = FCString::Atoi( *PositionValues[3] );
+			WindowX = FCString::Atoi(*PositionValues[0]);
+			WindowY = FCString::Atoi(*PositionValues[1]);
+			WindowWidth = FCString::Atoi(*PositionValues[2]);
+			WindowHeight = FCString::Atoi(*PositionValues[3]);
 		}
 	}
 
-	GConfig->GetFloat( TEXT("FoliageEdit"), TEXT("Radius"), Radius, GEditorUserSettingsIni );
-	GConfig->GetFloat( TEXT("FoliageEdit"), TEXT("PaintDensity"), PaintDensity, GEditorUserSettingsIni );
-	GConfig->GetFloat( TEXT("FoliageEdit"), TEXT("UnpaintDensity"), UnpaintDensity, GEditorUserSettingsIni );
-	GConfig->GetBool( TEXT("FoliageEdit"), TEXT("bFilterLandscape"), bFilterLandscape, GEditorUserSettingsIni );
-	GConfig->GetBool( TEXT("FoliageEdit"), TEXT("bFilterStaticMesh"), bFilterStaticMesh, GEditorUserSettingsIni );
-	GConfig->GetBool( TEXT("FoliageEdit"), TEXT("bFilterBSP"), bFilterBSP, GEditorUserSettingsIni );
+	GConfig->GetFloat(TEXT("FoliageEdit"), TEXT("Radius"), Radius, GEditorUserSettingsIni);
+	GConfig->GetFloat(TEXT("FoliageEdit"), TEXT("PaintDensity"), PaintDensity, GEditorUserSettingsIni);
+	GConfig->GetFloat(TEXT("FoliageEdit"), TEXT("UnpaintDensity"), UnpaintDensity, GEditorUserSettingsIni);
+	GConfig->GetBool(TEXT("FoliageEdit"), TEXT("bFilterLandscape"), bFilterLandscape, GEditorUserSettingsIni);
+	GConfig->GetBool(TEXT("FoliageEdit"), TEXT("bFilterStaticMesh"), bFilterStaticMesh, GEditorUserSettingsIni);
+	GConfig->GetBool(TEXT("FoliageEdit"), TEXT("bFilterBSP"), bFilterBSP, GEditorUserSettingsIni);
 }
 
 /** Save UI settings to ini file */
 void FFoliageUISettings::Save()
 {
-	FString WindowPositionString = FString::Printf(TEXT("%d,%d,%d,%d"), WindowX, WindowY, WindowWidth, WindowHeight );
-	GConfig->SetString( TEXT("FoliageEdit"), TEXT("WindowPosition"), *WindowPositionString, GEditorUserSettingsIni );
-
-	GConfig->SetFloat( TEXT("FoliageEdit"), TEXT("Radius"), Radius, GEditorUserSettingsIni );
-	GConfig->SetFloat( TEXT("FoliageEdit"), TEXT("PaintDensity"), PaintDensity, GEditorUserSettingsIni );
-	GConfig->SetFloat( TEXT("FoliageEdit"), TEXT("UnpaintDensity"), UnpaintDensity, GEditorUserSettingsIni );
-	GConfig->SetBool( TEXT("FoliageEdit"), TEXT("bFilterLandscape"), bFilterLandscape, GEditorUserSettingsIni );
-	GConfig->SetBool( TEXT("FoliageEdit"), TEXT("bFilterStaticMesh"), bFilterStaticMesh, GEditorUserSettingsIni );
-	GConfig->SetBool( TEXT("FoliageEdit"), TEXT("bFilterBSP"), bFilterBSP, GEditorUserSettingsIni );
+	FString WindowPositionString = FString::Printf(TEXT("%d,%d,%d,%d"), WindowX, WindowY, WindowWidth, WindowHeight);
+	GConfig->SetString(TEXT("FoliageEdit"), TEXT("WindowPosition"), *WindowPositionString, GEditorUserSettingsIni);
+
+	GConfig->SetFloat(TEXT("FoliageEdit"), TEXT("Radius"), Radius, GEditorUserSettingsIni);
+	GConfig->SetFloat(TEXT("FoliageEdit"), TEXT("PaintDensity"), PaintDensity, GEditorUserSettingsIni);
+	GConfig->SetFloat(TEXT("FoliageEdit"), TEXT("UnpaintDensity"), UnpaintDensity, GEditorUserSettingsIni);
+	GConfig->SetBool(TEXT("FoliageEdit"), TEXT("bFilterLandscape"), bFilterLandscape, GEditorUserSettingsIni);
+	GConfig->SetBool(TEXT("FoliageEdit"), TEXT("bFilterStaticMesh"), bFilterStaticMesh, GEditorUserSettingsIni);
+	GConfig->SetBool(TEXT("FoliageEdit"), TEXT("bFilterBSP"), bFilterBSP, GEditorUserSettingsIni);
 }
