Initial support for other foliage types

diff --git a/Engine/Source/Editor/FoliageEdit/Private/FoliageEdMode.h b/Engine/Source/Editor/FoliageEdit/Private/FoliageEdMode.h
--- a/Engine/Source/Editor/FoliageEdit/Private/FoliageEdMode.h
+++ b/Engine/Source/Editor/FoliageEdit/Private/FoliageEdMode.h
@@ -4,12 +4,12 @@
 	FoliageEdMode.h: Foliage editing
 ================================================================================*/
 
-
-#ifndef __FoliageEdMode_h__
-#define __FoliageEdMode_h__
-
 #pragma once
 
+// Forward declarations
+class ULandscapeComponent;
+class UFoliageType;
+
 #include "InstancedFoliage.h"
 
 //
@@ -47,29 +47,29 @@ struct FFoliageUISettings
 	float GetUnpaintDensity() const { return UnpaintDensity; }
 	void SetUnpaintDensity(float InUnpaintDensity) { UnpaintDensity = InUnpaintDensity; }
 	bool GetFilterLandscape() const { return bFilterLandscape ? true : false; }
-	void SetFilterLandscape(bool InbFilterLandscape ) { bFilterLandscape = InbFilterLandscape; }
+	void SetFilterLandscape(bool InbFilterLandscape) { bFilterLandscape = InbFilterLandscape; }
 	bool GetFilterStaticMesh() const { return bFilterStaticMesh ? true : false; }
-	void SetFilterStaticMesh(bool InbFilterStaticMesh ) { bFilterStaticMesh = InbFilterStaticMesh; }
+	void SetFilterStaticMesh(bool InbFilterStaticMesh) { bFilterStaticMesh = InbFilterStaticMesh; }
 	bool GetFilterBSP() const { return bFilterBSP ? true : false; }
-	void SetFilterBSP(bool InbFilterBSP ) { bFilterBSP = InbFilterBSP; }
+	void SetFilterBSP(bool InbFilterBSP) { bFilterBSP = InbFilterBSP; }
 
 	FFoliageUISettings()
-	:	WindowX(-1)
-	,	WindowY(-1)
-	,	WindowWidth(284)
-	,	WindowHeight(400)
-	,	bPaintToolSelected(true)
-	,	bReapplyToolSelected(false)
-	,	bSelectToolSelected(false)
-	,	bLassoSelectToolSelected(false)
-	,	bPaintBucketToolSelected(false)
-	,	bReapplyPaintBucketToolSelected(false)
-	,	Radius(512.f)
-	,	PaintDensity(0.5f)
-	,	UnpaintDensity(0.f)
-	,	bFilterLandscape(true)
-	,	bFilterStaticMesh(true)
-	,	bFilterBSP(true)
+		: WindowX(-1)
+		, WindowY(-1)
+		, WindowWidth(284)
+		, WindowHeight(400)
+		, bPaintToolSelected(true)
+		, bReapplyToolSelected(false)
+		, bSelectToolSelected(false)
+		, bLassoSelectToolSelected(false)
+		, bPaintBucketToolSelected(false)
+		, bReapplyPaintBucketToolSelected(false)
+		, Radius(512.f)
+		, PaintDensity(0.5f)
+		, UnpaintDensity(0.f)
+		, bFilterLandscape(true)
+		, bFilterStaticMesh(true)
+		, bFilterBSP(true)
 	{
 	}
 
@@ -103,12 +103,12 @@ public:
 
 struct FFoliageMeshUIInfo
 {
-	UStaticMesh* StaticMesh;
+	UFoliageType* Settings;
 	FFoliageMeshInfo* MeshInfo;
 
-	FFoliageMeshUIInfo(UStaticMesh* InStaticMesh, FFoliageMeshInfo* InMeshInfo)
-	:	StaticMesh(InStaticMesh)
-	,	MeshInfo(InMeshInfo)
+	FFoliageMeshUIInfo(UFoliageType* InSettings, FFoliageMeshInfo* InMeshInfo)
+		: Settings(InSettings)
+		, MeshInfo(InMeshInfo)
 	{}
 };
 
@@ -119,28 +119,28 @@ class FMeshInfoSnapshot
 	TArray<FVector> Locations;
 public:
 	FMeshInfoSnapshot(FFoliageMeshInfo* MeshInfo)
-	:	Hash(*MeshInfo->InstanceHash)
+		: Hash(*MeshInfo->InstanceHash)
 	{
 		Locations.AddUninitialized(MeshInfo->Instances.Num());
-		for( int32 Idx=0;Idx<MeshInfo->Instances.Num();Idx++ )
+		for (int32 Idx = 0; Idx < MeshInfo->Instances.Num(); Idx++)
 		{
 			Locations[Idx] = MeshInfo->Instances[Idx].Location;
 		}
 	}
 
-	int32 CountInstancesInsideSphere( const FSphere& Sphere )
+	int32 CountInstancesInsideSphere(const FSphere& Sphere)
 	{
-		TSet<int32> TempInstances;
-		Hash.GetInstancesOverlappingBox(FBox::BuildAABB(Sphere.Center, FVector(Sphere.W,Sphere.W,Sphere.W)), TempInstances );
+		int32 Count = 0;
 
-		int32 Count=0;
-		for( TSet<int32>::TConstIterator It(TempInstances); It; ++It )
+		auto TempInstances = Hash.GetInstancesOverlappingBox(FBox::BuildAABB(Sphere.Center, FVector(Sphere.W)));
+		for (int32 Idx : TempInstances)
 		{
-			if( FSphere(Locations[*It],0.f).IsInside(Sphere) )
+			if (FSphere(Locations[Idx], 0.f).IsInside(Sphere))
 			{
 				Count++;
 			}
-		}	
+		}
+
 		return Count;
 	}
 };
@@ -161,7 +161,7 @@ public:
 	virtual ~FEdModeFoliage();
 
 	/** FGCObject interface */
-	virtual void AddReferencedObjects( FReferenceCollector& Collector ) override;
+	virtual void AddReferencedObjects(FReferenceCollector& Collector) override;
 
 	/** FEdMode: Called when the mode is created */
 	virtual void Initialize() override;
@@ -197,7 +197,7 @@ public:
 	 *
 	 * @return	true if input was handled
 	 */
-	virtual bool MouseMove( FEditorViewportClient* ViewportClient, FViewport* Viewport, int32 x, int32 y ) override;
+	virtual bool MouseMove(FEditorViewportClient* ViewportClient, FViewport* Viewport, int32 x, int32 y) override;
 
 	/**
 	 * FEdMode: Called when the mouse is moved while a window input capture is in effect
@@ -209,7 +209,7 @@ public:
 	 *
 	 * @return	true if input was handled
 	 */
-	virtual bool CapturedMouseMove( FEditorViewportClient* InViewportClient, FViewport* InViewport, int32 InMouseX, int32 InMouseY ) override;
+	virtual bool CapturedMouseMove(FEditorViewportClient* InViewportClient, FViewport* InViewport, int32 InMouseX, int32 InMouseY) override;
 
 	/** FEdMode: Called when a mouse button is pressed */
 	virtual bool StartTracking(FEditorViewportClient* InViewportClient, FViewport* InViewport) override;
@@ -218,110 +218,111 @@ public:
 	virtual bool EndTracking(FEditorViewportClient* InViewportClient, FViewport* InViewport) override;
 
 	/** FEdMode: Called once per frame */
-	virtual void Tick(FEditorViewportClient* ViewportClient,float DeltaTime) override;
+	virtual void Tick(FEditorViewportClient* ViewportClient, float DeltaTime) override;
 
 	/** FEdMode: Called when a key is pressed */
-	virtual bool InputKey( FEditorViewportClient* InViewportClient, FViewport* InViewport, FKey InKey, EInputEvent InEvent ) override;
+	virtual bool InputKey(FEditorViewportClient* InViewportClient, FViewport* InViewport, FKey InKey, EInputEvent InEvent) override;
 
 	/** FEdMode: Called when mouse drag input it applied */
-	virtual bool InputDelta( FEditorViewportClient* InViewportClient, FViewport* InViewport, FVector& InDrag, FRotator& InRot, FVector& InScale ) override;
+	virtual bool InputDelta(FEditorViewportClient* InViewportClient, FViewport* InViewport, FVector& InDrag, FRotator& InRot, FVector& InScale) override;
 
 	/** FEdMode: Render elements for the Foliage tool */
-	virtual void Render( const FSceneView* View, FViewport* Viewport, FPrimitiveDrawInterface* PDI ) override;
+	virtual void Render(const FSceneView* View, FViewport* Viewport, FPrimitiveDrawInterface* PDI) override;
 
 	/** FEdMode: Render HUD elements for this tool */
-	virtual void DrawHUD( FEditorViewportClient* ViewportClient, FViewport* Viewport, const FSceneView* View, FCanvas* Canvas ) override;
+	virtual void DrawHUD(FEditorViewportClient* ViewportClient, FViewport* Viewport, const FSceneView* View, FCanvas* Canvas) override;
 
 	/** FEdMode: Handling SelectActor */
-	virtual bool Select( AActor* InActor, bool bInSelected ) override;
+	virtual bool Select(AActor* InActor, bool bInSelected) override;
 
 	/** FEdMode: Check to see if an actor can be selected in this mode - no side effects */
-	virtual bool IsSelectionAllowed( AActor* InActor, bool bInSelection ) const override;
+	virtual bool IsSelectionAllowed(AActor* InActor, bool bInSelection) const override;
 
 	/** FEdMode: Called when the currently selected actor has changed */
 	virtual void ActorSelectionChangeNotify() override;
 
 	/** Notifies all active modes of mouse click messages. */
-	bool HandleClick(FEditorViewportClient* InViewportClient,  HHitProxy *HitProxy, const FViewportClick &Click );
+	bool HandleClick(FEditorViewportClient* InViewportClient, HHitProxy *HitProxy, const FViewportClick &Click);
 
 	/** FEdMode: widget handling */
 	virtual FVector GetWidgetLocation() const override;
 	virtual bool AllowWidgetMove();
 	virtual bool ShouldDrawWidget() const override;
-	virtual bool UsesTransformWidget() const override ;
-	virtual EAxisList::Type GetWidgetAxisToDraw( FWidget::EWidgetMode InWidgetMode ) const override;
+	virtual bool UsesTransformWidget() const override;
+	virtual EAxisList::Type GetWidgetAxisToDraw(FWidget::EWidgetMode InWidgetMode) const override;
 
 	virtual bool DisallowMouseDeltaTracking() const override;
 
 	/** Forces real-time perspective viewports */
-	void ForceRealTimeViewports( const bool bEnable, const bool bStoreCurrentState );
+	void ForceRealTimeViewports(const bool bEnable, const bool bStoreCurrentState);
 
 	/** Trace under the mouse cursor and update brush position */
-	void FoliageBrushTrace( FEditorViewportClient* ViewportClient, int32 MouseX, int32 MouseY );
+	void FoliageBrushTrace(FEditorViewportClient* ViewportClient, int32 MouseX, int32 MouseY);
 
-	/** Generate start/end points for a random trace inside the sphere brush. 
+	/** Generate start/end points for a random trace inside the sphere brush.
 	    returns a line segment inside the sphere parallel to the view direction */
-	void GetRandomVectorInBrush(FVector& OutStart, FVector& OutEnd );
+	void GetRandomVectorInBrush(FVector& OutStart, FVector& OutEnd);
 
 	/** Apply brush */
-	void ApplyBrush( FEditorViewportClient* ViewportClient );
+	void ApplyBrush(FEditorViewportClient* ViewportClient);
 
 	/** Update existing mesh info for current level */
 	void UpdateFoliageMeshList();
 
 	TArray<struct FFoliageMeshUIInfo>& GetFoliageMeshList();
-	
+
 	/** Add a new mesh */
-	void AddFoliageMesh(UStaticMesh* StaticMesh);
+	void AddFoliageMesh(UStaticMesh* Mesh);
+	void AddFoliageMesh(UFoliageType* Settings);
 
 	/** Remove a mesh */
-	bool RemoveFoliageMesh(UStaticMesh* StaticMesh);
+	bool RemoveFoliageMesh(UFoliageType* Settings);
 
 	/** Reapply cluster settings to all the instances */
-	void ReallocateClusters(UStaticMesh* StaticMesh);
+	void ReallocateClusters(UFoliageType* Settings);
 
 	/** Replace a mesh with another one */
-	bool ReplaceStaticMesh(UStaticMesh* OldStaticMesh, UStaticMesh* NewStaticMesh, bool& bOutMeshMerged);
+	bool ReplaceStaticMesh(UFoliageType* Settings, UStaticMesh* NewStaticMesh, bool& bOutMeshMerged);
 
 	/** Bake meshes to StaticMeshActors */
-	void BakeFoliage(UStaticMesh* StaticMesh, bool bSelectedOnly);
+	void BakeFoliage(UFoliageType* Settings, bool bSelectedOnly);
 
-	/** 
+	/**
 	 * Copy the settings object for this static mesh
 	 *
 	 * @param StaticMesh	The static mesh to copy the settings of.
 	 *
 	 * @return The duplicated settings now assigned to the static mesh.
 	 */
-	UInstancedFoliageSettings* CopySettingsObject(UStaticMesh* StaticMesh);
+	UFoliageType* CopySettingsObject(UFoliageType* Settings);
 
 	/** Replace the settings object for this static mesh with the one specified */
-	void ReplaceSettingsObject(UStaticMesh* StaticMesh, UInstancedFoliageSettings* NewSettings);
+	void ReplaceSettingsObject(UFoliageType* OldSettings, UFoliageType* NewSettings);
 
-	/** 
-	 * Save the settings object 
+	/**
+	 * Save the settings object
 	 *
 	 * @param InSettingsPackageName		The name for the settings package.
 	 * @param StaticMesh				The static mesh to save the settings of.
 	 *
 	 * @return The settings the static mesh is now using.
 	 */
-	UInstancedFoliageSettings* SaveSettingsObject(const FText& InSettingsPackageName, UStaticMesh* StaticMesh);
+	UFoliageType* SaveSettingsObject(const FText& InSettingsPackageName, UFoliageType* Settings);
 
 	/** Apply paint bucket to actor */
 	void ApplyPaintBucket(AActor* Actor, bool bRemove);
 private:
 	/** Add instances inside the brush to match DesiredInstanceCount */
-	void AddInstancesForBrush( UWorld* InWorld, AInstancedFoliageActor* IFA, UStaticMesh* StaticMesh, FFoliageMeshInfo& MeshInfo, int32 DesiredInstanceCount, TArray<int32>& ExistingInstances, float Pressure );
+	void AddInstancesForBrush(UWorld* InWorld, AInstancedFoliageActor* IFA, UFoliageType* Settings, FFoliageMeshInfo& MeshInfo, int32 DesiredInstanceCount, TArray<int32>& ExistingInstances, float Pressure);
 
 	/** Remove instances inside the brush to match DesiredInstanceCount */
-	void RemoveInstancesForBrush( AInstancedFoliageActor* IFA, FFoliageMeshInfo& MeshInfo, int32 DesiredInstanceCount, TArray<int32>& ExistingInstances, float Pressure );
+	void RemoveInstancesForBrush(AInstancedFoliageActor* IFA, FFoliageMeshInfo& MeshInfo, int32 DesiredInstanceCount, TArray<int32>& ExistingInstances, float Pressure);
 
 	/** Reapply instance settings to exiting instances */
-	void ReapplyInstancesForBrush( UWorld* InWorld, AInstancedFoliageActor* IFA, FFoliageMeshInfo& MeshInfo, TArray<int32>& ExistingInstances );
+	void ReapplyInstancesForBrush(UWorld* InWorld, AInstancedFoliageActor* IFA, UFoliageType* Settings, TArray<int32>& ExistingInstances);
 
 	/** Lookup the vertex color corresponding to a location traced on a static mesh */
-	bool GetStaticMeshVertexColorForHit( UStaticMeshComponent* InStaticMeshComponent, int32 InTriangleIndex, const FVector& InHitLocation, FColor& OutVertexColor ) const;
+	bool GetStaticMeshVertexColorForHit(UStaticMeshComponent* InStaticMeshComponent, int32 InTriangleIndex, const FVector& InHitLocation, FColor& OutVertexColor) const;
 
 	bool bBrushTraceValid;
 	FVector BrushLocation;
@@ -329,13 +330,13 @@ private:
 	UStaticMeshComponent* SphereBrushComponent;
 
 	// Landscape layer cache data
-	TMap<FName, TMap<class ULandscapeComponent*, TArray<uint8> > > LandscapeLayerCaches;
+	TMap<FName, TMap<ULandscapeComponent*, TArray<uint8> > > LandscapeLayerCaches;
 
 	// Placed level data
-	TArray<struct FFoliageMeshUIInfo> FoliageMeshList;
+	TArray<FFoliageMeshUIInfo> FoliageMeshList;
 
 	// Cache of instance positions at the start of the transaction
-	TMap<UStaticMesh*, class FMeshInfoSnapshot> InstanceSnapshot;
+	TMap<UFoliageType*, FMeshInfoSnapshot> InstanceSnapshot;
 
 	bool bToolActive;
 	bool bCanAltDrag;
@@ -343,6 +344,3 @@ private:
 	// The actor that owns any currently visible selections
 	AInstancedFoliageActor* SelectionIFA;
 };
-
-
-#endif	// __FoliageEdMode_h__
diff --git a/Engine/Source/Editor/FoliageEdit/Private/SFoliageEdit.cpp b/Engine/Source/Editor/FoliageEdit/Private/SFoliageEdit.cpp
--- a/Engine/Source/Editor/FoliageEdit/Private/SFoliageEdit.cpp
+++ b/Engine/Source/Editor/FoliageEdit/Private/SFoliageEdit.cpp
@@ -5,7 +5,7 @@
 #include "SFoliageEdit.h"
 #include "FoliageEditActions.h"
 #include "FoliageEdMode.h"
-#include "Runtime/Engine/Classes/Foliage/InstancedFoliageSettings.h"
+#include "Foliage/FoliageType.h"
 #include "Editor/UnrealEd/Public/AssetThumbnail.h"
 
 #include "Editor/UnrealEd/Public/DragAndDrop/AssetDragDropOp.h"
@@ -36,7 +36,7 @@ public:
 
 	~SFoliageDragDropHandler() {}
 
-	FReply OnDrop( const FGeometry& MyGeometry, const FDragDropEvent& DragDropEvent )
+	FReply OnDrop(const FGeometry& MyGeometry, const FDragDropEvent& DragDropEvent)
 	{
 		TSharedPtr<SFoliageEdit> PinnedPtr = FoliageEditPtr.Pin();
 		if (PinnedPtr.IsValid())
@@ -47,7 +47,7 @@ public:
 		return FReply::Handled();
 	}
 
-	virtual void OnDragEnter( const FGeometry& MyGeometry, const FDragDropEvent& DragDropEvent ) override
+	virtual void OnDragEnter(const FGeometry& MyGeometry, const FDragDropEvent& DragDropEvent) override
 	{
 		TSharedPtr<SFoliageEdit> PinnedPtr = FoliageEditPtr.Pin();
 		if (PinnedPtr.IsValid())
@@ -57,7 +57,7 @@ public:
 		SCompoundWidget::OnDragEnter(MyGeometry, DragDropEvent);
 	}
 
-	virtual void OnDragLeave( const FDragDropEvent& DragDropEvent ) override
+	virtual void OnDragLeave(const FDragDropEvent& DragDropEvent) override
 	{
 		TSharedPtr<SFoliageEdit> PinnedPtr = FoliageEditPtr.Pin();
 		if (PinnedPtr.IsValid())
@@ -74,17 +74,17 @@ private:
 BEGIN_SLATE_FUNCTION_BUILD_OPTIMIZATION
 void SFoliageEdit::Construct(const FArguments& InArgs)
 {
-	FoliageEditMode = (FEdModeFoliage*)GLevelEditorModeTools().GetActiveMode( FBuiltinEditorModes::EM_Foliage );
+	FoliageEditMode = (FEdModeFoliage*)GLevelEditorModeTools().GetActiveMode(FBuiltinEditorModes::EM_Foliage);
 
 	FFoliageEditCommands::Register();
 
-	UICommandList = MakeShareable( new FUICommandList );
+	UICommandList = MakeShareable(new FUICommandList);
 
 	BindCommands();
 
 	FEditorDelegates::NewCurrentLevel.AddSP(this, &SFoliageEdit::NotifyNewCurrentLevel);
 
-	AssetThumbnailPool = MakeShareable( new FAssetThumbnailPool(512, TAttribute<bool>::Create( TAttribute<bool>::FGetter::CreateSP(this, &SFoliageEdit::IsHovered) ) ) );
+	AssetThumbnailPool = MakeShareable(new FAssetThumbnailPool(512, TAttribute<bool>::Create(TAttribute<bool>::FGetter::CreateSP(this, &SFoliageEdit::IsHovered))));
 
 	bOverlayOverride = false;
 
@@ -92,226 +92,226 @@ void SFoliageEdit::Construct(const FArguments& InArgs)
 	FMargin StandardPadding(0.0f, 4.0f, 0.0f, 4.0f);
 
 	this->ChildSlot
+	[
+		SNew(SScrollBox)
+		+ SScrollBox::Slot()
+		.Padding(0.0f)
 		[
-			SNew( SScrollBox )
-			+ SScrollBox::Slot()
-			.Padding( 0.0f )
+			SNew(SVerticalBox)
+
+			+ SVerticalBox::Slot()
+			.AutoHeight()
+			.HAlign(HAlign_Center)
+			.Padding(StandardPadding)
 			[
-				SNew( SVerticalBox )
-				 
-				+SVerticalBox::Slot()
-				.AutoHeight()
-				.HAlign(HAlign_Center)
+				BuildToolBar()
+			]
+
+			+ SVerticalBox::Slot()
+			.AutoHeight()
+			[
+				SNew(SHorizontalBox)
+
+				+ SHorizontalBox::Slot()
+				.FillWidth(1.0f)
 				.Padding(StandardPadding)
 				[
-					BuildToolBar()
+					SNew(STextBlock)
+					.Visibility(this, &SFoliageEdit::GetVisibility_Radius)
+					.Text(LOCTEXT("BrushSize", "Brush Size"))
 				]
-
-				+SVerticalBox::Slot()
-				.AutoHeight()
+				+ SHorizontalBox::Slot()
+				.FillWidth(2.0f)
+				.Padding(StandardPadding)
 				[
-					SNew( SHorizontalBox )
-
-					+SHorizontalBox::Slot()
-						.FillWidth(1.0f)
-						.Padding(StandardPadding)
-					[
-						SNew(STextBlock)
-							.Visibility( this, &SFoliageEdit::GetVisibility_Radius )
-							.Text(LOCTEXT("BrushSize", "Brush Size"))
-					]
-					+SHorizontalBox::Slot()
-						.FillWidth(2.0f)
-						.Padding(StandardPadding)
-					[
-						SNew(SSpinBox<float>)
-							.Visibility( this, &SFoliageEdit::GetVisibility_Radius )
-							.MinValue(0.0f)
-							.MaxValue(65536.0f)
-							.MaxSliderValue(8192.0f)
-							.SliderExponent(3.0f)
-							.Value(this, &SFoliageEdit::GetRadius)
-							.OnValueChanged(this, &SFoliageEdit::SetRadius)
-					]
+					SNew(SSpinBox<float>)
+					.Visibility(this, &SFoliageEdit::GetVisibility_Radius)
+					.MinValue(0.0f)
+					.MaxValue(65536.0f)
+					.MaxSliderValue(8192.0f)
+					.SliderExponent(3.0f)
+					.Value(this, &SFoliageEdit::GetRadius)
+					.OnValueChanged(this, &SFoliageEdit::SetRadius)
 				]
+			]
 
-				+SVerticalBox::Slot()
-				.AutoHeight()
-				[
-					SNew( SHorizontalBox )
+			+ SVerticalBox::Slot()
+			.AutoHeight()
+			[
+				SNew(SHorizontalBox)
 
-					+SHorizontalBox::Slot()
-						.FillWidth(1.0f)
-						.Padding(StandardPadding)
-					[
-						SNew(STextBlock)
-							.Visibility( this, &SFoliageEdit::GetVisibility_PaintDensity )
-							.Text(LOCTEXT("PaintDensity", "Paint Density"))
-					]
-					+SHorizontalBox::Slot()
-						.FillWidth(2.0f)
-						.Padding(StandardPadding)
-					[
-						SNew(SSpinBox<float>)
-							.Visibility( this, &SFoliageEdit::GetVisibility_PaintDensity )
-							.MinValue(0.0f)
-							.MaxValue(2.0f)
-							.MaxSliderValue(1.0f)
-							.Value(this, &SFoliageEdit::GetPaintDensity)
-							.OnValueChanged(this, &SFoliageEdit::SetPaintDensity)
-					]
+				+ SHorizontalBox::Slot()
+				.FillWidth(1.0f)
+				.Padding(StandardPadding)
+				[
+					SNew(STextBlock)
+					.Visibility(this, &SFoliageEdit::GetVisibility_PaintDensity)
+					.Text(LOCTEXT("PaintDensity", "Paint Density"))
 				]
-			
-				+SVerticalBox::Slot()
-				.AutoHeight()
+				+ SHorizontalBox::Slot()
+				.FillWidth(2.0f)
+				.Padding(StandardPadding)
 				[
-					SNew( SHorizontalBox )
-
-					+SHorizontalBox::Slot()
-						.FillWidth(1.0f)
-						.Padding(StandardPadding)
-					[
-						SNew(STextBlock)
-							.Visibility( this, &SFoliageEdit::GetVisibility_EraseDensity )
-							.Text(LOCTEXT("EraseDensity", "Erase Density"))
-					]
-					+SHorizontalBox::Slot()
-						.FillWidth(2.0f)
-						.Padding(StandardPadding)
-					[
-						SNew(SSpinBox<float>)
-							.Visibility( this, &SFoliageEdit::GetVisibility_EraseDensity )
-							.MinValue(0.0f)
-							.MaxValue(2.0f)
-							.MaxSliderValue(1.0f)
-							.Value(this, &SFoliageEdit::GetEraseDensity)
-							.OnValueChanged(this, &SFoliageEdit::SetEraseDensity)
-					]
+					SNew(SSpinBox<float>)
+					.Visibility(this, &SFoliageEdit::GetVisibility_PaintDensity)
+					.MinValue(0.0f)
+					.MaxValue(2.0f)
+					.MaxSliderValue(1.0f)
+					.Value(this, &SFoliageEdit::GetPaintDensity)
+					.OnValueChanged(this, &SFoliageEdit::SetPaintDensity)
 				]
-		
-				+SVerticalBox::Slot()
-				.AutoHeight()
+			]
+
+			+ SVerticalBox::Slot()
+			.AutoHeight()
+			[
+				SNew(SHorizontalBox)
+
+				+ SHorizontalBox::Slot()
+				.FillWidth(1.0f)
 				.Padding(StandardPadding)
 				[
 					SNew(STextBlock)
-						.Visibility( this, &SFoliageEdit::GetVisibility_Filters )
-						.Text(LOCTEXT("Filter", "Filter"))
+					.Visibility(this, &SFoliageEdit::GetVisibility_EraseDensity)
+					.Text(LOCTEXT("EraseDensity", "Erase Density"))
 				]
-
-				+SVerticalBox::Slot()
-				.AutoHeight()
+				+ SHorizontalBox::Slot()
+				.FillWidth(2.0f)
 				.Padding(StandardPadding)
 				[
-					SNew(SWrapBox)
-						.UseAllottedWidth(true)
+					SNew(SSpinBox<float>)
+					.Visibility(this, &SFoliageEdit::GetVisibility_EraseDensity)
+					.MinValue(0.0f)
+					.MaxValue(2.0f)
+					.MaxSliderValue(1.0f)
+					.Value(this, &SFoliageEdit::GetEraseDensity)
+					.OnValueChanged(this, &SFoliageEdit::SetEraseDensity)
+				]
+			]
+
+			+ SVerticalBox::Slot()
+			.AutoHeight()
+			.Padding(StandardPadding)
+			[
+				SNew(STextBlock)
+				.Visibility(this, &SFoliageEdit::GetVisibility_Filters)
+				.Text(LOCTEXT("Filter", "Filter"))
+			]
+
+			+ SVerticalBox::Slot()
+			.AutoHeight()
+			.Padding(StandardPadding)
+			[
+				SNew(SWrapBox)
+				.UseAllottedWidth(true)
 
-					+SWrapBox::Slot()
-						.Padding(0.0f, 0.0f, 16.0f, 0.0f)
+				+ SWrapBox::Slot()
+				.Padding(0.0f, 0.0f, 16.0f, 0.0f)
+				[
+					SNew(SCheckBox)
+					.Visibility(this, &SFoliageEdit::GetVisibility_Filters)
+					.OnCheckStateChanged(this, &SFoliageEdit::OnCheckStateChanged_Landscape)
+					.IsChecked(this, &SFoliageEdit::GetCheckState_Landscape)
 					[
-						SNew(SCheckBox)
-							.Visibility( this, &SFoliageEdit::GetVisibility_Filters )
-							.OnCheckStateChanged( this, &SFoliageEdit::OnCheckStateChanged_Landscape )
-							.IsChecked( this, &SFoliageEdit::GetCheckState_Landscape )
-							[
-								SNew(STextBlock)
-									.Text(LOCTEXT("Landscape", "Landscape"))
-							]
+						SNew(STextBlock)
+						.Text(LOCTEXT("Landscape", "Landscape"))
 					]
-				
-					+SWrapBox::Slot()
-						.Padding(0.0f, 0.0f, 16.0f, 0.0f)
+				]
+
+				+ SWrapBox::Slot()
+				.Padding(0.0f, 0.0f, 16.0f, 0.0f)
+				[
+					SNew(SCheckBox)
+					.Visibility(this, &SFoliageEdit::GetVisibility_Filters)
+					.OnCheckStateChanged(this, &SFoliageEdit::OnCheckStateChanged_StaticMesh)
+					.IsChecked(this, &SFoliageEdit::GetCheckState_StaticMesh)
 					[
-						SNew(SCheckBox)
-							.Visibility( this, &SFoliageEdit::GetVisibility_Filters )
-							.OnCheckStateChanged( this, &SFoliageEdit::OnCheckStateChanged_StaticMesh )
-							.IsChecked( this, &SFoliageEdit::GetCheckState_StaticMesh )
-							[
-								SNew(STextBlock)
-									.Text(LOCTEXT("StaticMeshes", "Static Meshes"))
-							]
+						SNew(STextBlock)
+						.Text(LOCTEXT("StaticMeshes", "Static Meshes"))
 					]
+				]
 
-					+SWrapBox::Slot()
+				+ SWrapBox::Slot()
+				[
+					SNew(SCheckBox)
+					.Visibility(this, &SFoliageEdit::GetVisibility_Filters)
+					.OnCheckStateChanged(this, &SFoliageEdit::OnCheckStateChanged_BSP)
+					.IsChecked(this, &SFoliageEdit::GetCheckState_BSP)
 					[
-						SNew(SCheckBox)
-							.Visibility( this, &SFoliageEdit::GetVisibility_Filters )
-							.OnCheckStateChanged( this, &SFoliageEdit::OnCheckStateChanged_BSP )
-							.IsChecked( this, &SFoliageEdit::GetCheckState_BSP )
-							[
-								SNew(STextBlock)
-									.Text(LOCTEXT("BSP", "BSP"))
-							]
+						SNew(STextBlock)
+						.Text(LOCTEXT("BSP", "BSP"))
 					]
 				]
+			]
 
-				+SVerticalBox::Slot()
-				.AutoHeight()
-				[
-					SNew(SHeaderRow)
+			+ SVerticalBox::Slot()
+			.AutoHeight()
+			[
+				SNew(SHeaderRow)
 
-					+SHeaderRow::Column(TEXT("Meshes"))
-						.DefaultLabel(LOCTEXT("Meshes", "Meshes"))
-				]
+				+ SHeaderRow::Column(TEXT("Meshes"))
+				.DefaultLabel(LOCTEXT("Meshes", "Meshes"))
+			]
 
-				+SVerticalBox::Slot()
-				.AutoHeight()
-				.Padding(StandardPadding)
+			+ SVerticalBox::Slot()
+			.AutoHeight()
+			.Padding(StandardPadding)
+			[
+
+				SNew(SOverlay)
+
+				+ SOverlay::Slot()
 				[
+					SNew(SFoliageDragDropHandler)
+					.FoliageEditPtr(SharedThis(this))
+					[
+						SAssignNew(ItemScrollBox, SScrollBox)
+					]
+				]
 
-					SNew(SOverlay)
+				+ SOverlay::Slot()
+				[
+					SNew(SHorizontalBox)
+					.Visibility(this, &SFoliageEdit::GetVisibility_EmptyList)
 
-					+SOverlay::Slot()
+					+ SHorizontalBox::Slot()
+					.FillWidth(1.0f)
 					[
-						SNew(SFoliageDragDropHandler)
-							.FoliageEditPtr(SharedThis(this))
+						SNew(SBorder)
+						.BorderImage(FEditorStyle::GetBrush("Menu.Background"))
+						.BorderBackgroundColor(FLinearColor(1, 1, 1, 0.8f))
+						.Content()
+						[
+							SNew(SHorizontalBox)
+
+							+ SHorizontalBox::Slot()
+							.FillWidth(1.0f)
 							[
-								SAssignNew(ItemScrollBox, SScrollBox)
+								SNew(SSpacer)
 							]
-					]
 
-					+SOverlay::Slot()
-					[
-						SNew(SHorizontalBox)
-							.Visibility( this, &SFoliageEdit::GetVisibility_EmptyList )
+							+ SHorizontalBox::Slot()
+							.AutoWidth()
+							.VAlign(VAlign_Center)
+							[
+								SNew(STextBlock)
+								.Visibility(this, &SFoliageEdit::GetVisibility_EmptyListText)
+								.Text(LOCTEXT("EmptyFoliageListMessage", "Drag static meshes from the content browser into this area."))
+							]
 
-						+SHorizontalBox::Slot()
+							+ SHorizontalBox::Slot()
 							.FillWidth(1.0f)
-						[
-							SNew(SBorder)
-							.BorderImage( FEditorStyle::GetBrush("Menu.Background") )
-							.BorderBackgroundColor(FLinearColor(1, 1, 1, 0.8f))
-							.Content()
 							[
-								SNew(SHorizontalBox)
-
-								+SHorizontalBox::Slot()
-								.FillWidth(1.0f)
-								[
-									SNew(SSpacer)
-								]
-
-								+SHorizontalBox::Slot()
-									.AutoWidth()
-									.VAlign(VAlign_Center)
-								[
-									SNew(STextBlock)
-										.Visibility( this, &SFoliageEdit::GetVisibility_EmptyListText )
-										.Text( LOCTEXT("EmptyFoliageListMessage", "Drag static meshes from the content browser into this area.") )	
-								]
-
-								+SHorizontalBox::Slot()
-									.FillWidth(1.0f)
-									[
-										SNew(SSpacer)
-									]
+								SNew(SSpacer)
 							]
 						]
 					]
 				]
 			]
-		];
+		]
+	];
 
-		RefreshFullList();
+	RefreshFullList();
 }
 END_SLATE_FUNCTION_BUILD_OPTIMIZATION
 
@@ -319,51 +319,52 @@ SFoliageEdit::~SFoliageEdit()
 {
 	// Release all rendering resources being held onto
 	AssetThumbnailPool->ReleaseResources();
-	
+
 	FEditorDelegates::NewCurrentLevel.RemoveAll(this);
 }
 
-void SFoliageEdit::Tick( const FGeometry& AllottedGeometry, const double InCurrentTime, const float InDeltaTime )
+void SFoliageEdit::Tick(const FGeometry& AllottedGeometry, const double InCurrentTime, const float InDeltaTime)
 {
 	AssetThumbnailPool->Tick(InDeltaTime);
 }
 
 TSharedPtr<FAssetThumbnail> SFoliageEdit::CreateThumbnail(UStaticMesh* InStaticMesh)
 {
-	return MakeShareable( new FAssetThumbnail( InStaticMesh, 80, 80, AssetThumbnailPool ) );
+	return MakeShareable(new FAssetThumbnail(InStaticMesh, 80, 80, AssetThumbnailPool));
 }
 
 void SFoliageEdit::RefreshFullList()
 {
-	for( int DisplayItemIdx = 0; DisplayItemIdx < DisplayItemList.Num(); )
+	for (int DisplayItemIdx = 0; DisplayItemIdx < DisplayItemList.Num();)
 	{
 		RemoveItemFromScrollbox(DisplayItemList[DisplayItemIdx]);
 	}
 
 	TArray<struct FFoliageMeshUIInfo>& FoliageMeshList = FoliageEditMode->GetFoliageMeshList();
 
-	for( int MeshIdx = 0; MeshIdx < FoliageMeshList.Num(); MeshIdx++ )
+	for (int MeshIdx = 0; MeshIdx < FoliageMeshList.Num(); MeshIdx++)
 	{
-		AddItemToScrollbox( FoliageMeshList[MeshIdx] );
+		AddItemToScrollbox(FoliageMeshList[MeshIdx]);
 	}
 }
 
 void SFoliageEdit::AddItemToScrollbox(struct FFoliageMeshUIInfo& InFoliageInfoToAdd)
 {
-	TSharedRef< SFoliageEditMeshDisplayItem > DisplayItem = 
-		SNew( SFoliageEditMeshDisplayItem )
-			.FoliageEditPtr( SharedThis(this) )
-			.FoliageSettingsPtr( InFoliageInfoToAdd.MeshInfo->Settings )
-			.AssociatedStaticMesh(InFoliageInfoToAdd.StaticMesh)
-			.AssetThumbnail( CreateThumbnail(InFoliageInfoToAdd.StaticMesh) )
-			.FoliageMeshUIInfo(MakeShareable(new FFoliageMeshUIInfo( InFoliageInfoToAdd ) ));
+	UFoliageType* Settings = InFoliageInfoToAdd.Settings;
+
+	TSharedRef<SFoliageEditMeshDisplayItem> DisplayItem =
+		SNew(SFoliageEditMeshDisplayItem)
+		.FoliageEditPtr(SharedThis(this))
+		.FoliageSettingsPtr(Settings)
+		.AssetThumbnail(CreateThumbnail(Settings->GetStaticMesh()))
+		.FoliageMeshUIInfo(MakeShareable(new FFoliageMeshUIInfo(InFoliageInfoToAdd)));
 
 	DisplayItemList.Add(DisplayItem);
 
 	ItemScrollBox->AddSlot()
-	[
-		DisplayItem			
-	];
+		[
+			DisplayItem
+		];
 }
 
 void SFoliageEdit::RemoveItemFromScrollbox(const TSharedPtr<SFoliageEditMeshDisplayItem> InWidgetToRemove)
@@ -373,19 +374,19 @@ void SFoliageEdit::RemoveItemFromScrollbox(const TSharedPtr<SFoliageEditMeshDisp
 	ItemScrollBox->RemoveSlot(InWidgetToRemove.ToSharedRef());
 }
 
-void SFoliageEdit::ReplaceItem(const TSharedPtr<class SFoliageEditMeshDisplayItem> InDisplayItemToReplaceIn, UStaticMesh* InNewStaticMesh)
+void SFoliageEdit::ReplaceItem(const TSharedPtr<SFoliageEditMeshDisplayItem> InDisplayItemToReplaceIn, UStaticMesh* InNewStaticMesh)
 {
-	TArray<struct FFoliageMeshUIInfo>& FoliageMeshList = FoliageEditMode->GetFoliageMeshList();
+	TArray<FFoliageMeshUIInfo>& FoliageMeshList = FoliageEditMode->GetFoliageMeshList();
 
-	for(int32 UIInfoIndex = 0; UIInfoIndex < FoliageMeshList.Num(); ++UIInfoIndex)
+	for (FFoliageMeshUIInfo& UIInfo : FoliageMeshList)
 	{
-		if(FoliageMeshList[UIInfoIndex].StaticMesh == InNewStaticMesh)
+		UFoliageType* Settings = UIInfo.Settings;
+		if (Settings->GetStaticMesh() == InNewStaticMesh)
 		{
 			// This is the info for the new static mesh. Update the display item.
-			InDisplayItemToReplaceIn->Replace(FoliageMeshList[UIInfoIndex].MeshInfo->Settings, 
-												FoliageMeshList[UIInfoIndex].StaticMesh, 
-												CreateThumbnail(FoliageMeshList[UIInfoIndex].StaticMesh), 
-												MakeShareable(new FFoliageMeshUIInfo( FoliageMeshList[UIInfoIndex] ) ) );
+			InDisplayItemToReplaceIn->Replace(Settings,
+				CreateThumbnail(Settings->GetStaticMesh()),
+				MakeShareable(new FFoliageMeshUIInfo(UIInfo)));
 			break;
 		}
 	}
@@ -406,38 +407,38 @@ void SFoliageEdit::BindCommands()
 
 	UICommandList->MapAction(
 		Commands.SetPaint,
-		FExecuteAction::CreateSP( this, &SFoliageEdit::OnSetPaint ),
+		FExecuteAction::CreateSP(this, &SFoliageEdit::OnSetPaint),
 		FCanExecuteAction(),
-		FIsActionChecked::CreateSP( this, &SFoliageEdit::IsPaintTool ) );
+		FIsActionChecked::CreateSP(this, &SFoliageEdit::IsPaintTool));
 
 	UICommandList->MapAction(
 		Commands.SetReapplySettings,
-		FExecuteAction::CreateSP( this, &SFoliageEdit::OnSetReapplySettings ),
+		FExecuteAction::CreateSP(this, &SFoliageEdit::OnSetReapplySettings),
 		FCanExecuteAction(),
-		FIsActionChecked::CreateSP( this, &SFoliageEdit::IsReapplySettingsTool ) );
+		FIsActionChecked::CreateSP(this, &SFoliageEdit::IsReapplySettingsTool));
 
 	UICommandList->MapAction(
 		Commands.SetSelect,
-		FExecuteAction::CreateSP( this, &SFoliageEdit::OnSetSelectInstance ),
+		FExecuteAction::CreateSP(this, &SFoliageEdit::OnSetSelectInstance),
 		FCanExecuteAction(),
-		FIsActionChecked::CreateSP( this, &SFoliageEdit::IsSelectTool ) );
+		FIsActionChecked::CreateSP(this, &SFoliageEdit::IsSelectTool));
 
 	UICommandList->MapAction(
 		Commands.SetLassoSelect,
-		FExecuteAction::CreateSP( this, &SFoliageEdit::OnSetLasso ),
+		FExecuteAction::CreateSP(this, &SFoliageEdit::OnSetLasso),
 		FCanExecuteAction(),
-		FIsActionChecked::CreateSP( this, &SFoliageEdit::IsLassoSelectTool ) );
+		FIsActionChecked::CreateSP(this, &SFoliageEdit::IsLassoSelectTool));
 
 	UICommandList->MapAction(
 		Commands.SetPaintBucket,
-		FExecuteAction::CreateSP( this, &SFoliageEdit::OnSetPaintFill ),
+		FExecuteAction::CreateSP(this, &SFoliageEdit::OnSetPaintFill),
 		FCanExecuteAction(),
-		FIsActionChecked::CreateSP( this, &SFoliageEdit::IsPaintFillTool ) );
+		FIsActionChecked::CreateSP(this, &SFoliageEdit::IsPaintFillTool));
 }
 
 TSharedRef<SWidget> SFoliageEdit::BuildToolBar()
 {
-	FToolBarBuilder Toolbar(UICommandList, FMultiBoxCustomization::None );
+	FToolBarBuilder Toolbar(UICommandList, FMultiBoxCustomization::None);
 	{
 		Toolbar.AddToolBarButton(FFoliageEditCommands::Get().SetPaint);
 		Toolbar.AddToolBarButton(FFoliageEditCommands::Get().SetReapplySettings);
@@ -449,14 +450,14 @@ TSharedRef<SWidget> SFoliageEdit::BuildToolBar()
 	return
 		SNew(SHorizontalBox)
 
-		+SHorizontalBox::Slot()
+		+ SHorizontalBox::Slot()
 		.AutoWidth()
-		.Padding(4,0)
+		.Padding(4, 0)
 		[
 			SNew(SBorder)
 			.Padding(0)
 			.BorderImage(FEditorStyle::GetBrush("NoBorder"))
-			.IsEnabled( FSlateApplication::Get().GetNormalExecutionAttribute() )
+			.IsEnabled(FSlateApplication::Get().GetNormalExecutionAttribute())
 			[
 				Toolbar.MakeWidget()
 			]
@@ -563,19 +564,19 @@ float SFoliageEdit::GetEraseDensity() const
 	return FoliageEditMode->UISettings.GetUnpaintDensity();
 }
 
-void SFoliageEdit::OnDragEnter_ListView( const FGeometry& MyGeometry, const FDragDropEvent& DragDropEvent )
+void SFoliageEdit::OnDragEnter_ListView(const FGeometry& MyGeometry, const FDragDropEvent& DragDropEvent)
 {
 	bOverlayOverride = true;
 	SCompoundWidget::OnDragEnter(MyGeometry, DragDropEvent);
 }
 
-void SFoliageEdit::OnDragLeave_ListView( const FDragDropEvent& DragDropEvent )
+void SFoliageEdit::OnDragLeave_ListView(const FDragDropEvent& DragDropEvent)
 {
 	bOverlayOverride = false;
 	SCompoundWidget::OnDragLeave(DragDropEvent);
 }
 
-FReply SFoliageEdit::OnDragOver_ListView( const FGeometry& MyGeometry, const FDragDropEvent& DragDropEvent )
+FReply SFoliageEdit::OnDragOver_ListView(const FGeometry& MyGeometry, const FDragDropEvent& DragDropEvent)
 {
 
 	return FReply::Handled();
@@ -583,17 +584,17 @@ FReply SFoliageEdit::OnDragOver_ListView( const FGeometry& MyGeometry, const FDr
 
 EVisibility SFoliageEdit::GetVisibility_EmptyList() const
 {
-	if( DisplayItemList.Num() > 0 && !bOverlayOverride )
+	if (DisplayItemList.Num() > 0 && !bOverlayOverride)
 	{
 		return EVisibility::Collapsed;
 	}
-	
+
 	return EVisibility::HitTestInvisible;
 }
 
 EVisibility SFoliageEdit::GetVisibility_EmptyListText() const
 {
-	if( DisplayItemList.Num() > 0 )
+	if (DisplayItemList.Num() > 0)
 	{
 		return EVisibility::Collapsed;
 	}
@@ -603,7 +604,7 @@ EVisibility SFoliageEdit::GetVisibility_EmptyListText() const
 
 EVisibility SFoliageEdit::GetVisibility_NonEmptyList() const
 {
-	if( DisplayItemList.Num() == 0 )
+	if (DisplayItemList.Num() == 0)
 	{
 		return EVisibility::Collapsed;
 	}
@@ -613,9 +614,11 @@ EVisibility SFoliageEdit::GetVisibility_NonEmptyList() const
 
 bool SFoliageEdit::CanAddStaticMesh(const UStaticMesh* const InStaticMesh) const
 {
-	for(int32 FoliageIdx = 0; FoliageIdx < FoliageEditMode->GetFoliageMeshList().Num(); ++FoliageIdx)
+	for (FFoliageMeshUIInfo& UIInfo : FoliageEditMode->GetFoliageMeshList())
 	{
-		if( FoliageEditMode->GetFoliageMeshList()[FoliageIdx].StaticMesh == InStaticMesh )
+		UFoliageType* Settings = UIInfo.Settings;
+
+		if (Settings->GetStaticMesh() == InStaticMesh)
 		{
 			return false;
 		}
@@ -623,17 +626,17 @@ bool SFoliageEdit::CanAddStaticMesh(const UStaticMesh* const InStaticMesh) const
 	return true;
 }
 
-FReply SFoliageEdit::OnDrop_ListView( const FGeometry& MyGeometry, const FDragDropEvent& DragDropEvent )
+FReply SFoliageEdit::OnDrop_ListView(const FGeometry& MyGeometry, const FDragDropEvent& DragDropEvent)
 {
-	TArray< FAssetData > DroppedAssetData = AssetUtil::ExtractAssetDataFromDrag( DragDropEvent );
+	TArray<FAssetData> DroppedAssetData = AssetUtil::ExtractAssetDataFromDrag(DragDropEvent);
 
 	for (int32 AssetIdx = 0; AssetIdx < DroppedAssetData.Num(); ++AssetIdx)
 	{
-		GWarn->BeginSlowTask( LOCTEXT("OnDrop_LoadPackage", "Fully Loading Package For Drop"), true, false );
-		UStaticMesh* StaticMesh = Cast<UStaticMesh>( DroppedAssetData[ AssetIdx ].GetAsset() );
+		GWarn->BeginSlowTask(LOCTEXT("OnDrop_LoadPackage", "Fully Loading Package For Drop"), true, false);
+		UStaticMesh* StaticMesh = Cast<UStaticMesh>(DroppedAssetData[AssetIdx].GetAsset());
 		GWarn->EndSlowTask();
 
-		if( StaticMesh && CanAddStaticMesh(StaticMesh) )
+		if (StaticMesh && CanAddStaticMesh(StaticMesh))
 		{
 			FoliageEditMode->AddFoliageMesh(StaticMesh);
 			AddItemToScrollbox(FoliageEditMode->GetFoliageMeshList()[FoliageEditMode->GetFoliageMeshList().Num() - 1]);
@@ -646,37 +649,37 @@ FReply SFoliageEdit::OnDrop_ListView( const FGeometry& MyGeometry, const FDragDr
 
 void SFoliageEdit::OnCheckStateChanged_Landscape(ESlateCheckBoxState::Type InState)
 {
-	FoliageEditMode->UISettings.SetFilterLandscape(InState == ESlateCheckBoxState::Checked? true : false);
+	FoliageEditMode->UISettings.SetFilterLandscape(InState == ESlateCheckBoxState::Checked ? true : false);
 }
 
 ESlateCheckBoxState::Type SFoliageEdit::GetCheckState_Landscape() const
 {
-	return FoliageEditMode->UISettings.GetFilterLandscape()? ESlateCheckBoxState::Checked : ESlateCheckBoxState::Unchecked;
+	return FoliageEditMode->UISettings.GetFilterLandscape() ? ESlateCheckBoxState::Checked : ESlateCheckBoxState::Unchecked;
 }
 
 void SFoliageEdit::OnCheckStateChanged_StaticMesh(ESlateCheckBoxState::Type InState)
 {
-	FoliageEditMode->UISettings.SetFilterStaticMesh(InState == ESlateCheckBoxState::Checked? true : false);
+	FoliageEditMode->UISettings.SetFilterStaticMesh(InState == ESlateCheckBoxState::Checked ? true : false);
 }
 
 ESlateCheckBoxState::Type SFoliageEdit::GetCheckState_StaticMesh() const
 {
-	return FoliageEditMode->UISettings.GetFilterStaticMesh()? ESlateCheckBoxState::Checked : ESlateCheckBoxState::Unchecked;
+	return FoliageEditMode->UISettings.GetFilterStaticMesh() ? ESlateCheckBoxState::Checked : ESlateCheckBoxState::Unchecked;
 }
 
 void SFoliageEdit::OnCheckStateChanged_BSP(ESlateCheckBoxState::Type InState)
 {
-	FoliageEditMode->UISettings.SetFilterBSP(InState == ESlateCheckBoxState::Checked? true : false);
+	FoliageEditMode->UISettings.SetFilterBSP(InState == ESlateCheckBoxState::Checked ? true : false);
 }
 
 ESlateCheckBoxState::Type SFoliageEdit::GetCheckState_BSP() const
 {
-	return FoliageEditMode->UISettings.GetFilterBSP()? ESlateCheckBoxState::Checked : ESlateCheckBoxState::Unchecked;
+	return FoliageEditMode->UISettings.GetFilterBSP() ? ESlateCheckBoxState::Checked : ESlateCheckBoxState::Unchecked;
 }
 
 EVisibility SFoliageEdit::GetVisibility_Radius() const
 {
-	if( FoliageEditMode->UISettings.GetSelectToolSelected() || FoliageEditMode->UISettings.GetReapplyPaintBucketToolSelected() || FoliageEditMode->UISettings.GetPaintBucketToolSelected() )
+	if (FoliageEditMode->UISettings.GetSelectToolSelected() || FoliageEditMode->UISettings.GetReapplyPaintBucketToolSelected() || FoliageEditMode->UISettings.GetPaintBucketToolSelected())
 	{
 		return EVisibility::Collapsed;
 	}
@@ -686,7 +689,7 @@ EVisibility SFoliageEdit::GetVisibility_Radius() const
 
 EVisibility SFoliageEdit::GetVisibility_PaintDensity() const
 {
-	if( FoliageEditMode->UISettings.GetSelectToolSelected() || FoliageEditMode->UISettings.GetReapplyPaintBucketToolSelected() || FoliageEditMode->UISettings.GetLassoSelectToolSelected())
+	if (FoliageEditMode->UISettings.GetSelectToolSelected() || FoliageEditMode->UISettings.GetReapplyPaintBucketToolSelected() || FoliageEditMode->UISettings.GetLassoSelectToolSelected())
 	{
 		return EVisibility::Collapsed;
 	}
@@ -696,7 +699,7 @@ EVisibility SFoliageEdit::GetVisibility_PaintDensity() const
 
 EVisibility SFoliageEdit::GetVisibility_EraseDensity() const
 {
-	if( FoliageEditMode->UISettings.GetSelectToolSelected() || FoliageEditMode->UISettings.GetReapplyPaintBucketToolSelected() || FoliageEditMode->UISettings.GetLassoSelectToolSelected() || FoliageEditMode->UISettings.GetPaintBucketToolSelected() )
+	if (FoliageEditMode->UISettings.GetSelectToolSelected() || FoliageEditMode->UISettings.GetReapplyPaintBucketToolSelected() || FoliageEditMode->UISettings.GetLassoSelectToolSelected() || FoliageEditMode->UISettings.GetPaintBucketToolSelected())
 	{
 		return EVisibility::Collapsed;
 	}
@@ -706,7 +709,7 @@ EVisibility SFoliageEdit::GetVisibility_EraseDensity() const
 
 EVisibility SFoliageEdit::GetVisibility_Filters() const
 {
-	if( FoliageEditMode->UISettings.GetSelectToolSelected() )
+	if (FoliageEditMode->UISettings.GetSelectToolSelected())
 	{
 		return EVisibility::Collapsed;
 	}
@@ -714,14 +717,14 @@ EVisibility SFoliageEdit::GetVisibility_Filters() const
 	return EVisibility::Visible;
 }
 
-void SFoliageEdit::DisableDragDropOverlay(void) 
-{ 
-	bOverlayOverride = false; 
+void SFoliageEdit::DisableDragDropOverlay(void)
+{
+	bOverlayOverride = false;
 }
 
-void SFoliageEdit::EnableDragDropOverlay(void) 
-{ 
-	bOverlayOverride = true; 
+void SFoliageEdit::EnableDragDropOverlay(void)
+{
+	bOverlayOverride = true;
 }
 
 void SFoliageEdit::NotifyNewCurrentLevel()
diff --git a/Engine/Source/Editor/FoliageEdit/Private/SFoliageEdit.h b/Engine/Source/Editor/FoliageEdit/Private/SFoliageEdit.h
--- a/Engine/Source/Editor/FoliageEdit/Private/SFoliageEdit.h
+++ b/Engine/Source/Editor/FoliageEdit/Private/SFoliageEdit.h
@@ -2,6 +2,12 @@
 
 #pragma once
 
+// Forwards declarations
+class FAssetThumbnail;
+class SFoliageEditMeshDisplayItem;
+struct FFoliageMeshUIInfo;
+class FAssetThumbnailPool;
+
 class SFoliageEdit : public SCompoundWidget
 {
 public:
@@ -15,34 +21,34 @@ public:
 	~SFoliageEdit();
 
 	// SWidget interface
-	virtual void Tick( const FGeometry& AllottedGeometry, const double InCurrentTime, const float InDeltaTime ) override;
+	virtual void Tick(const FGeometry& AllottedGeometry, const double InCurrentTime, const float InDeltaTime) override;
 	// End of SWidget interface
 
 	/** Creates the thumbnail for the passed in Static Mesh. Used by the MeshListView items. */
-	TSharedPtr< class FAssetThumbnail > CreateThumbnail(UStaticMesh* InStaticMesh);
-	
+	TSharedPtr<FAssetThumbnail> CreateThumbnail(UStaticMesh* InStaticMesh);
+
 	/** Does a full refresh on the list. */
 	void RefreshFullList();
 
 	/** Adds a static mesh to the list of available meshes for foliage. May be called on by the MeshListView items. */
-	void AddItemToScrollbox(struct FFoliageMeshUIInfo& InFoliageInfoToAdd);
+	void AddItemToScrollbox(FFoliageMeshUIInfo& InFoliageInfoToAdd);
 
 	/** Removes a static mesh from the list of available meshes for foliage. May be called on by the MeshListView items. */
-	void RemoveItemFromScrollbox(const TSharedPtr<class SFoliageEditMeshDisplayItem> InWidgetToRemove);
+	void RemoveItemFromScrollbox(const TSharedPtr<SFoliageEditMeshDisplayItem> InWidgetToRemove);
 
-	void ReplaceItem(const TSharedPtr<class SFoliageEditMeshDisplayItem> InDisplayItemToReplaceIn, UStaticMesh* InNewStaticMesh);
+	void ReplaceItem(const TSharedPtr<SFoliageEditMeshDisplayItem> InDisplayItemToReplaceIn, UStaticMesh* InNewStaticMesh);
 
 	/** Handles adding a new item to the list and refreshing the list in it's entirety. */
-	FReply OnDrop_ListView( const FGeometry& MyGeometry, const FDragDropEvent& DragDropEvent );
+	FReply OnDrop_ListView(const FGeometry& MyGeometry, const FDragDropEvent& DragDropEvent);
 
 	/** Gets FoliageEditMode. Used by the cluster details to notify changes */
 	class FEdModeFoliage* GetFoliageEditMode() const { return FoliageEditMode; }
 
 	/** Sets the overlay that appears over all items in the list to be visible. */
-	void DisableDragDropOverlay(void);
-	
+	void DisableDragDropOverlay();
+
 	/** Sets the overlay that appears over all items in the list to be invisible. */
-	void EnableDragDropOverlay(void);
+	void EnableDragDropOverlay();
 
 private:
 
@@ -106,16 +112,16 @@ private:
 	float GetEraseDensity() const;
 
 	/** Creates the list item widget that displays the instance settings. */
-	TSharedRef< ITableRow > MakeWidgetFromOption( TSharedPtr<struct FFoliageMeshUIInfo> InItem, const TSharedRef< STableViewBase >& OwnerTable );
+	TSharedRef<ITableRow> MakeWidgetFromOption(TSharedPtr<FFoliageMeshUIInfo> InItem, const TSharedRef<STableViewBase>& OwnerTable);
 
 	/** Handles nothing? */
-	void OnDragEnter_ListView( const FGeometry& MyGeometry, const FDragDropEvent& DragDropEvent );
+	void OnDragEnter_ListView(const FGeometry& MyGeometry, const FDragDropEvent& DragDropEvent);
 
 	/** Handles nothing? */
-	void OnDragLeave_ListView( const FDragDropEvent& DragDropEvent );
-	
+	void OnDragLeave_ListView(const FDragDropEvent& DragDropEvent);
+
 	/** @todo - Have something occur to help the user out in knowing something is actually going to occur (besides the mouse changing) */
-	FReply OnDragOver_ListView( const FGeometry& MyGeometry, const FDragDropEvent& DragDropEvent );
+	FReply OnDragOver_ListView(const FGeometry& MyGeometry, const FDragDropEvent& DragDropEvent);
 
 	/** Sets the filter settings for if painting will occur on Landscapes. */
 	void OnCheckStateChanged_Landscape(ESlateCheckBoxState::Type InState);
@@ -131,7 +137,7 @@ private:
 
 	/** Sets the filter settings for if painting will occur on BSPs. */
 	void OnCheckStateChanged_BSP(ESlateCheckBoxState::Type InState);
-	
+
 	/** Retrieves the filter settings for painting on BSPs. */
 	ESlateCheckBoxState::Type GetCheckState_BSP() const;
 
@@ -149,8 +155,8 @@ private:
 
 	/** Checks if the paint density spinbox should appear. Dependant on the current tool being used. */
 	EVisibility GetVisibility_PaintDensity() const;
-	
-	/** Checks if the erase density spinbox should appear. Dependant on the current tool being used. */	
+
+	/** Checks if the erase density spinbox should appear. Dependant on the current tool being used. */
 	EVisibility GetVisibility_EraseDensity() const;
 
 	/** Checks if the filters should appear. Dependant on the current tool being used. */
@@ -170,26 +176,26 @@ private:
 
 private:
 	/** The list view object for displaying Static Meshes to use for foliage. */
-	TSharedPtr< SListView< TSharedPtr< struct FFoliageMeshUIInfo > > > MeshListView;
+	TSharedPtr<SListView<TSharedPtr<FFoliageMeshUIInfo>>> MeshListView;
 
 	/** List of Static Meshes being used for foliage, adapted from a list retrieved from Foliage Mode and items should only be added to it in NotifyChanged to keep information
 		accurate between it and the Foliage Mode. */
-	TArray< TSharedPtr< struct FFoliageMeshUIInfo > > MeshList;
+	TArray<TSharedPtr<FFoliageMeshUIInfo>> MeshList;
 
 	/** Pool for maintaining and rendering thumbnails */
-	TSharedPtr<class FAssetThumbnailPool> AssetThumbnailPool;
+	TSharedPtr<FAssetThumbnailPool> AssetThumbnailPool;
 
 	/** Command list for binding functions for the toolbar. */
 	TSharedPtr<FUICommandList> UICommandList;
 
 	/** Pointer to the foliage edit mode. */
-	class FEdModeFoliage* FoliageEditMode;
+	FEdModeFoliage* FoliageEditMode;
 
 	/** Scrollbox for slotting foliage items. */
 	TSharedPtr<SScrollBox> ItemScrollBox;
 
 	/** List of items currently being displayed. */
-	TArray< TSharedRef<class SFoliageEditMeshDisplayItem> > DisplayItemList;
+	TArray<TSharedRef<SFoliageEditMeshDisplayItem>> DisplayItemList;
 
 	/** Used to override the empty list overlay to be visible, not including the text. */
 	bool bOverlayOverride;
