Initial support for other foliage types

diff --git a/Engine/Source/Editor/FoliageEdit/Private/SFoliageEditMeshDisplayItem.cpp b/Engine/Source/Editor/FoliageEdit/Private/SFoliageEditMeshDisplayItem.cpp
--- a/Engine/Source/Editor/FoliageEdit/Private/SFoliageEditMeshDisplayItem.cpp
+++ b/Engine/Source/Editor/FoliageEdit/Private/SFoliageEditMeshDisplayItem.cpp
@@ -5,7 +5,7 @@
 #include "SFoliageEdit.h"
 #include "FoliageEditActions.h"
 #include "FoliageEdMode.h"
-#include "Runtime/Engine/Classes/Foliage/InstancedFoliageSettings.h"
+#include "Foliage/FoliageType.h"
 #include "Editor/UnrealEd/Public/AssetThumbnail.h"
 #include "Dialogs/DlgPickAssetPath.h"
 #include "Editor/PropertyEditor/Public/PropertyEditorModule.h"
@@ -18,48 +18,47 @@
 #define LOCTEXT_NAMESPACE "FoliageEd_Mode"
 
 BEGIN_SLATE_FUNCTION_BUILD_OPTIMIZATION
-void SFoliageEditMeshDisplayItem::Construct( const FArguments& InArgs )
+void SFoliageEditMeshDisplayItem::Construct(const FArguments& InArgs)
 {
 	BorderImage = FEditorStyle::GetBrush("FoliageEditMode.BubbleBorder");
 	BorderBackgroundColor = TAttribute<FSlateColor>(this, &SFoliageEditMeshDisplayItem::GetBorderColor);
- 
+
 	FoliageEditPtr = InArgs._FoliageEditPtr;
 	FoliageSettingsPtr = InArgs._FoliageSettingsPtr;
-	AssociatedStaticMesh = InArgs._AssociatedStaticMesh;
 	Thumbnail = InArgs._AssetThumbnail;
 	FoliageMeshUIInfo = InArgs._FoliageMeshUIInfo;
 
-	ThumbnailWidget = InArgs._AssetThumbnail->MakeThumbnailWidget();
+	ThumbnailWidget = Thumbnail->MakeThumbnailWidget();
 
 	// Create the command list and bind the commands.
-	UICommandList = MakeShareable( new FUICommandList );
+	UICommandList = MakeShareable(new FUICommandList);
 	BindCommands();
 
 	CurrentViewSettings = ECurrentViewSettings::ShowPaintSettings;
 
 	// Create the details panels for the clustering tab.
-	FPropertyEditorModule& PropertyEditorModule = FModuleManager::GetModuleChecked<FPropertyEditorModule>( "PropertyEditor" );
-	const FDetailsViewArgs DetailsViewArgs( false, false, true, false, true, this );
-	ClusterSettingsDetails = PropertyEditorModule.CreateDetailView( DetailsViewArgs );
+	FPropertyEditorModule& PropertyEditorModule = FModuleManager::GetModuleChecked<FPropertyEditorModule>("PropertyEditor");
+	const FDetailsViewArgs DetailsViewArgs(false, false, true, false, true, this);
+	ClusterSettingsDetails = PropertyEditorModule.CreateDetailView(DetailsViewArgs);
 
 	DetailsObjectList.Add(FoliageSettingsPtr);
 	ClusterSettingsDetails->SetIsPropertyVisibleDelegate(FIsPropertyVisible::CreateSP(this, &SFoliageEditMeshDisplayItem::IsPropertyVisible));
 	ClusterSettingsDetails->SetObjects(DetailsObjectList);
-	
+
 	// Everything (or almost) uses this padding, change it to expand the padding.
 	FMargin StandardPadding(2.0f, 2.0f, 2.0f, 2.0f);
 	FMargin StandardSidePadding(2.0f, 0.0f, 2.0f, 0.0f);
 
 	// Create the toolbar for the tab options.
-	FToolBarBuilder Toolbar(UICommandList, FMultiBoxCustomization::None );
+	FToolBarBuilder Toolbar(UICommandList, FMultiBoxCustomization::None);
 	{
 		Toolbar.AddToolBarButton(FFoliageEditCommands::Get().SetNoSettings);
 		Toolbar.AddToolBarButton(FFoliageEditCommands::Get().SetPaintSettings);
 		Toolbar.AddToolBarButton(FFoliageEditCommands::Get().SetClusterSettings);
 	}
 
-	SAssignNew( ThumbnailBox, SBox )
-		.WidthOverride( 80 ) .HeightOverride( 80 )
+	SAssignNew(ThumbnailBox, SBox)
+		.WidthOverride(80).HeightOverride(80)
 		[
 			SAssignNew(ThumbnailWidgetBorder, SBorder)
 			[
@@ -71,27 +70,28 @@ void SFoliageEditMeshDisplayItem::Construct( const FArguments& InArgs )
 	const float SPINBOX_PREFIX = 0.3f;
 	const float SPINBOX_SINGLE = 2.2f;
 
-	TSharedRef<SHorizontalBox> DensityBox = SNew(SHorizontalBox)
+	TSharedRef<SHorizontalBox> DensityBox =
+		SNew(SHorizontalBox)
 
-		+SHorizontalBox::Slot()
-			.FillWidth(1.0f)
+		+ SHorizontalBox::Slot()
+		.FillWidth(1.0f)
 		[
 			SNew(SVerticalBox)
 
 			// Density
-			+SVerticalBox::Slot()
+			+ SVerticalBox::Slot()
 			.AutoHeight()
 			.Padding(StandardPadding)
 			[
-				SNew( SHorizontalBox )
+				SNew(SHorizontalBox)
 				.Visibility(this, &SFoliageEditMeshDisplayItem::IsNotReapplySettingsVisible)
-				
-				+SHorizontalBox::Slot()
+
+				+ SHorizontalBox::Slot()
 				.FillWidth(MAIN_TITLE + SPINBOX_PREFIX)
 				[
-					SNew( SHorizontalBox )
+					SNew(SHorizontalBox)
 
-					+SHorizontalBox::Slot()
+					+ SHorizontalBox::Slot()
 					.AutoWidth()
 					.VAlign(VAlign_Center)
 					[
@@ -99,17 +99,17 @@ void SFoliageEditMeshDisplayItem::Construct( const FArguments& InArgs )
 						.Text(LOCTEXT("Density", "Density / 1Kuu"))
 					]
 
-					+SHorizontalBox::Slot()
+					+ SHorizontalBox::Slot()
 					.AutoWidth()
 					.VAlign(VAlign_Center)
 					[
 						SNew(STextBlock)
 						.Text(LOCTEXT("DensitySuperscript", "2"))
-						.Font(FSlateFontInfo(FPaths::EngineContentDir() / TEXT("Slate/Fonts/Roboto-Regular.ttf" ), 8))
+						.Font(FSlateFontInfo(FPaths::EngineContentDir() / TEXT("Slate/Fonts/Roboto-Regular.ttf"), 8))
 					]
 				]
 
-				+SHorizontalBox::Slot()
+				+ SHorizontalBox::Slot()
 				.FillWidth(SPINBOX_SINGLE)
 				[
 					SNew(SSpinBox<float>)
@@ -122,14 +122,14 @@ void SFoliageEditMeshDisplayItem::Construct( const FArguments& InArgs )
 			]
 
 			// Density Adjustment
-			+SVerticalBox::Slot()
+			+ SVerticalBox::Slot()
 			.AutoHeight()
 			.Padding(StandardPadding)
 			[
-				SNew( SHorizontalBox )
+				SNew(SHorizontalBox)
 				.Visibility(this, &SFoliageEditMeshDisplayItem::IsReapplySettingsVisible)
 
-				+SHorizontalBox::Slot()
+				+ SHorizontalBox::Slot()
 				.AutoWidth()
 				[
 					SNew(SCheckBox)
@@ -137,7 +137,7 @@ void SFoliageEditMeshDisplayItem::Construct( const FArguments& InArgs )
 					.IsChecked(this, &SFoliageEditMeshDisplayItem::IsDensityReapplyChecked)
 				]
 
-				+SHorizontalBox::Slot()
+				+ SHorizontalBox::Slot()
 				.FillWidth(MAIN_TITLE + SPINBOX_PREFIX)
 				.VAlign(VAlign_Center)
 				[
@@ -145,7 +145,7 @@ void SFoliageEditMeshDisplayItem::Construct( const FArguments& InArgs )
 					.Text(LOCTEXT("DensityAdjust", "Density Adjust"))
 				]
 
-				+SHorizontalBox::Slot()
+				+ SHorizontalBox::Slot()
 				.FillWidth(SPINBOX_SINGLE)
 				[
 					SNew(SSpinBox<float>)
@@ -155,17 +155,16 @@ void SFoliageEditMeshDisplayItem::Construct( const FArguments& InArgs )
 					.Value(this, &SFoliageEditMeshDisplayItem::GetDensity)
 					.OnValueChanged(this, &SFoliageEditMeshDisplayItem::OnDensityChanged)
 				]
-
 			]
 
 			// Radius
-			+SVerticalBox::Slot()
+			+ SVerticalBox::Slot()
 			.AutoHeight()
 			.Padding(StandardPadding)
 			[
-				SNew( SHorizontalBox )
+				SNew(SHorizontalBox)
 
-				+SHorizontalBox::Slot()
+				+ SHorizontalBox::Slot()
 				.AutoWidth()
 				[
 					SNew(SCheckBox)
@@ -174,7 +173,7 @@ void SFoliageEditMeshDisplayItem::Construct( const FArguments& InArgs )
 					.IsChecked(this, &SFoliageEditMeshDisplayItem::IsRadiusReapplyChecked)
 				]
 
-				+SHorizontalBox::Slot()
+				+ SHorizontalBox::Slot()
 				.FillWidth(MAIN_TITLE + SPINBOX_PREFIX)
 				.VAlign(VAlign_Center)
 				[
@@ -182,7 +181,7 @@ void SFoliageEditMeshDisplayItem::Construct( const FArguments& InArgs )
 					.Text(LOCTEXT("Radius", "Radius"))
 				]
 
-				+SHorizontalBox::Slot()
+				+ SHorizontalBox::Slot()
 				.FillWidth(SPINBOX_SINGLE)
 				[
 					SNew(SSpinBox<float>)
@@ -195,16 +194,17 @@ void SFoliageEditMeshDisplayItem::Construct( const FArguments& InArgs )
 			]
 		]
 
-		+SHorizontalBox::Slot()
+		+ SHorizontalBox::Slot()
 		.AutoWidth()
 		.Padding(StandardPadding)
 		[
-			ThumbnailBox.ToSharedRef()	
+			ThumbnailBox.ToSharedRef()
 		];
 
-	TSharedRef<SHorizontalBox> AlignToNormalBox = SNew( SHorizontalBox )
+	TSharedRef<SHorizontalBox> AlignToNormalBox =
+		SNew(SHorizontalBox)
 
-		+SHorizontalBox::Slot()
+		+ SHorizontalBox::Slot()
 		.AutoWidth()
 		[
 			SNew(SCheckBox)
@@ -213,7 +213,7 @@ void SFoliageEditMeshDisplayItem::Construct( const FArguments& InArgs )
 			.IsChecked(this, &SFoliageEditMeshDisplayItem::IsAlignToNormalReapplyChecked)
 		]
 
-		+SHorizontalBox::Slot()
+		+ SHorizontalBox::Slot()
 		.AutoWidth()
 		.VAlign(VAlign_Center)
 		[
@@ -226,18 +226,19 @@ void SFoliageEditMeshDisplayItem::Construct( const FArguments& InArgs )
 			]
 		];
 
-	TSharedRef<SHorizontalBox> MaxAngleBox = SNew( SHorizontalBox )
+	TSharedRef<SHorizontalBox> MaxAngleBox =
+		SNew(SHorizontalBox)
 		.Visibility(this, &SFoliageEditMeshDisplayItem::IsAlignToNormalVisible)
 
 		// Dummy Checkbox
-		+SHorizontalBox::Slot()
+		+ SHorizontalBox::Slot()
 		.AutoWidth()
 		[
 			SNew(SCheckBox)
 			.Visibility(this, &SFoliageEditMeshDisplayItem::IsReapplySettingsVisible_Dummy)
 		]
 
-		+SHorizontalBox::Slot()
+		+ SHorizontalBox::Slot()
 		.FillWidth(MAIN_TITLE + SPINBOX_PREFIX)
 		.VAlign(VAlign_Center)
 		[
@@ -245,12 +246,12 @@ void SFoliageEditMeshDisplayItem::Construct( const FArguments& InArgs )
 			.Text(LOCTEXT("MaxAngle", "Max Angle +/-"))
 		]
 
-		+SHorizontalBox::Slot()
+		+ SHorizontalBox::Slot()
 		.FillWidth(SPINBOX_SINGLE)
 		[
-			SNew( SHorizontalBox )
+			SNew(SHorizontalBox)
 
-			+SHorizontalBox::Slot()
+			+ SHorizontalBox::Slot()
 			.FillWidth(1.f)
 			[
 				SNew(SSpinBox<float>)
@@ -261,7 +262,7 @@ void SFoliageEditMeshDisplayItem::Construct( const FArguments& InArgs )
 				.OnValueChanged(this, &SFoliageEditMeshDisplayItem::OnMaxAngleChanged)
 			]
 
-			+SHorizontalBox::Slot()
+			+ SHorizontalBox::Slot()
 			.AutoWidth()
 			[
 				SNew(SCheckBox)
@@ -269,9 +270,10 @@ void SFoliageEditMeshDisplayItem::Construct( const FArguments& InArgs )
 			]
 		];
 
-	TSharedRef<SHorizontalBox> RandomYawBox = SNew( SHorizontalBox )
+	TSharedRef<SHorizontalBox> RandomYawBox =
+		SNew(SHorizontalBox)
 
-		+SHorizontalBox::Slot()
+		+ SHorizontalBox::Slot()
 		.AutoWidth()
 		[
 			SNew(SCheckBox)
@@ -280,7 +282,7 @@ void SFoliageEditMeshDisplayItem::Construct( const FArguments& InArgs )
 			.IsChecked(this, &SFoliageEditMeshDisplayItem::IsRandomYawReapplyChecked)
 		]
 
-		+SHorizontalBox::Slot()
+		+ SHorizontalBox::Slot()
 		.AutoWidth()
 		.VAlign(VAlign_Center)
 		[
@@ -293,10 +295,11 @@ void SFoliageEditMeshDisplayItem::Construct( const FArguments& InArgs )
 			]
 		];
 
-	TSharedRef<SHorizontalBox> UniformScaleBox = SNew( SHorizontalBox )
+	TSharedRef<SHorizontalBox> UniformScaleBox =
+		SNew(SHorizontalBox)
 
 		// Dummy Checkbox
-		+SHorizontalBox::Slot()
+		+ SHorizontalBox::Slot()
 		.AutoWidth()
 		[
 			SNew(SCheckBox)
@@ -304,7 +307,7 @@ void SFoliageEditMeshDisplayItem::Construct( const FArguments& InArgs )
 		]
 
 		// Uniform Scale Checkbox
-		+SHorizontalBox::Slot()
+		+ SHorizontalBox::Slot()
 		.AutoWidth()
 		.VAlign(VAlign_Center)
 		[
@@ -317,13 +320,13 @@ void SFoliageEditMeshDisplayItem::Construct( const FArguments& InArgs )
 			]
 		]
 
-		+SHorizontalBox::Slot()
-			.FillWidth(1.0f)
+		+ SHorizontalBox::Slot()
+		.FillWidth(1.0f)
 		[
 			SNew(SSpacer)
 		]
 		// "Lock" Non-uniform scaling checkbox.
-		+SHorizontalBox::Slot()
+		+ SHorizontalBox::Slot()
 		.AutoWidth()
 		.VAlign(VAlign_Center)
 		[
@@ -332,10 +335,11 @@ void SFoliageEditMeshDisplayItem::Construct( const FArguments& InArgs )
 			.Visibility(this, &SFoliageEditMeshDisplayItem::IsNonUniformScalingVisible)
 		];
 
-	TSharedRef<SHorizontalBox> ScaleUniformBox = SNew( SHorizontalBox )
+	TSharedRef<SHorizontalBox> ScaleUniformBox =
+		SNew(SHorizontalBox)
 		.Visibility(this, &SFoliageEditMeshDisplayItem::IsUniformScalingVisible)
 
-		+SHorizontalBox::Slot()
+		+ SHorizontalBox::Slot()
 		.AutoWidth()
 		[
 			SNew(SCheckBox)
@@ -344,7 +348,7 @@ void SFoliageEditMeshDisplayItem::Construct( const FArguments& InArgs )
 			.IsChecked(this, &SFoliageEditMeshDisplayItem::IsScaleUniformReapplyChecked)
 		]
 
-		+SHorizontalBox::Slot()
+		+ SHorizontalBox::Slot()
 		.FillWidth(MAIN_TITLE)
 		.VAlign(VAlign_Center)
 		[
@@ -352,12 +356,12 @@ void SFoliageEditMeshDisplayItem::Construct( const FArguments& InArgs )
 			.Text(LOCTEXT("ScaleUniform", "Scale"))
 		]
 
-		+SHorizontalBox::Slot()
+		+ SHorizontalBox::Slot()
 		.FillWidth(SPINBOX_PREFIX)
 		[
-			SNew( SHorizontalBox )
-			+SHorizontalBox::Slot().FillWidth(1.0f) // empty slot to eat the space for the alignment
-			+SHorizontalBox::Slot()
+			SNew(SHorizontalBox)
+			+ SHorizontalBox::Slot().FillWidth(1.0f) // empty slot to eat the space for the alignment
+			+ SHorizontalBox::Slot()
 			.AutoWidth()
 			.HAlign(HAlign_Right)
 			.VAlign(VAlign_Center)
@@ -368,12 +372,12 @@ void SFoliageEditMeshDisplayItem::Construct( const FArguments& InArgs )
 			]
 		]
 
-		+SHorizontalBox::Slot()
+		+ SHorizontalBox::Slot()
 		.FillWidth(SPINBOX_SINGLE)
 		[
-			SNew( SHorizontalBox )
+			SNew(SHorizontalBox)
 
-			+SHorizontalBox::Slot()
+			+ SHorizontalBox::Slot()
 			.FillWidth(0.5)
 			[
 				SNew(SSpinBox<float>)
@@ -384,7 +388,7 @@ void SFoliageEditMeshDisplayItem::Construct( const FArguments& InArgs )
 				.OnValueChanged(this, &SFoliageEditMeshDisplayItem::OnScaleUniformMinChanged)
 			]
 
-			+SHorizontalBox::Slot()
+			+ SHorizontalBox::Slot()
 			.AutoWidth()
 			.VAlign(VAlign_Center)
 			.Padding(StandardSidePadding)
@@ -393,10 +397,10 @@ void SFoliageEditMeshDisplayItem::Construct( const FArguments& InArgs )
 				.Text(LOCTEXT("Max", "Max"))
 			]
 
-			+SHorizontalBox::Slot()
+			+ SHorizontalBox::Slot()
 			.FillWidth(0.5)
 			[
-				SNew(SSpinBox<float>)	
+				SNew(SSpinBox<float>)
 				.MinValue(0.0f)
 				.MaxValue(100.0f)
 				.MaxSliderValue(10.0f)
@@ -404,7 +408,7 @@ void SFoliageEditMeshDisplayItem::Construct( const FArguments& InArgs )
 				.OnValueChanged(this, &SFoliageEditMeshDisplayItem::OnScaleUniformMaxChanged)
 			]
 
-			+SHorizontalBox::Slot()
+			+ SHorizontalBox::Slot()
 			.AutoWidth()
 			[
 				SNew(SCheckBox)
@@ -412,10 +416,11 @@ void SFoliageEditMeshDisplayItem::Construct( const FArguments& InArgs )
 			]
 		];
 
-	TSharedRef<SHorizontalBox> ScaleXBox = SNew( SHorizontalBox )
+	TSharedRef<SHorizontalBox> ScaleXBox =
+		SNew(SHorizontalBox)
 		.Visibility(this, &SFoliageEditMeshDisplayItem::IsNonUniformScalingVisible)
 
-		+SHorizontalBox::Slot()
+		+ SHorizontalBox::Slot()
 		.AutoWidth()
 		[
 			SNew(SCheckBox)
@@ -424,7 +429,7 @@ void SFoliageEditMeshDisplayItem::Construct( const FArguments& InArgs )
 			.IsChecked(this, &SFoliageEditMeshDisplayItem::IsScaleXReapplyChecked)
 		]
 
-		+SHorizontalBox::Slot()
+		+ SHorizontalBox::Slot()
 		.FillWidth(MAIN_TITLE)
 		.VAlign(VAlign_Center)
 		[
@@ -432,12 +437,12 @@ void SFoliageEditMeshDisplayItem::Construct( const FArguments& InArgs )
 			.Text(LOCTEXT("ScaleX", "Scale X"))
 		]
 
-		+SHorizontalBox::Slot()
+		+ SHorizontalBox::Slot()
 		.FillWidth(SPINBOX_PREFIX)
 		[
-			SNew( SHorizontalBox )
-			+SHorizontalBox::Slot().FillWidth(1.0f) // empty slot to eat the space for the alignment
-			+SHorizontalBox::Slot()
+			SNew(SHorizontalBox)
+			+ SHorizontalBox::Slot().FillWidth(1.0f) // empty slot to eat the space for the alignment
+			+ SHorizontalBox::Slot()
 			.AutoWidth()
 			.HAlign(HAlign_Right)
 			.VAlign(VAlign_Center)
@@ -448,12 +453,12 @@ void SFoliageEditMeshDisplayItem::Construct( const FArguments& InArgs )
 			]
 		]
 
-		+SHorizontalBox::Slot()
+		+ SHorizontalBox::Slot()
 		.FillWidth(SPINBOX_SINGLE)
 		[
-			SNew( SHorizontalBox )
+			SNew(SHorizontalBox)
 
-			+SHorizontalBox::Slot()
+			+ SHorizontalBox::Slot()
 			.FillWidth(0.5f)
 			[
 				SNew(SSpinBox<float>)
@@ -464,7 +469,7 @@ void SFoliageEditMeshDisplayItem::Construct( const FArguments& InArgs )
 				.OnValueChanged(this, &SFoliageEditMeshDisplayItem::OnScaleXMinChanged)
 			]
 
-			+SHorizontalBox::Slot()
+			+ SHorizontalBox::Slot()
 			.AutoWidth()
 			.VAlign(VAlign_Center)
 			.Padding(StandardSidePadding)
@@ -473,10 +478,10 @@ void SFoliageEditMeshDisplayItem::Construct( const FArguments& InArgs )
 				.Text(LOCTEXT("Max", "Max"))
 			]
 
-			+SHorizontalBox::Slot()
+			+ SHorizontalBox::Slot()
 			.FillWidth(0.5f)
 			[
-				SNew(SSpinBox<float>)	
+				SNew(SSpinBox<float>)
 				.MinValue(0.0f)
 				.MaxValue(100.0f)
 				.MaxSliderValue(10.0f)
@@ -484,7 +489,7 @@ void SFoliageEditMeshDisplayItem::Construct( const FArguments& InArgs )
 				.OnValueChanged(this, &SFoliageEditMeshDisplayItem::OnScaleXMaxChanged)
 			]
 
-			+SHorizontalBox::Slot()
+			+ SHorizontalBox::Slot()
 			.AutoWidth()
 			[
 				SNew(SCheckBox)
@@ -493,10 +498,11 @@ void SFoliageEditMeshDisplayItem::Construct( const FArguments& InArgs )
 			]
 		];
 
-	TSharedRef<SHorizontalBox> ScaleYBox = SNew( SHorizontalBox )
+	TSharedRef<SHorizontalBox> ScaleYBox =
+		SNew(SHorizontalBox)
 		.Visibility(this, &SFoliageEditMeshDisplayItem::IsNonUniformScalingVisible)
 
-		+SHorizontalBox::Slot()
+		+ SHorizontalBox::Slot()
 		.AutoWidth()
 		[
 			SNew(SCheckBox)
@@ -505,7 +511,7 @@ void SFoliageEditMeshDisplayItem::Construct( const FArguments& InArgs )
 			.IsChecked(this, &SFoliageEditMeshDisplayItem::IsScaleYReapplyChecked)
 		]
 
-		+SHorizontalBox::Slot()
+		+ SHorizontalBox::Slot()
 		.FillWidth(MAIN_TITLE)
 		.VAlign(VAlign_Center)
 		[
@@ -513,12 +519,12 @@ void SFoliageEditMeshDisplayItem::Construct( const FArguments& InArgs )
 			.Text(LOCTEXT("ScaleY", "Scale Y"))
 		]
 
-		+SHorizontalBox::Slot()
+		+ SHorizontalBox::Slot()
 		.FillWidth(SPINBOX_PREFIX)
 		[
-			SNew( SHorizontalBox )
-			+SHorizontalBox::Slot().FillWidth(1.0f) // empty slot to eat the space for the alignment
-			+SHorizontalBox::Slot()
+			SNew(SHorizontalBox)
+			+ SHorizontalBox::Slot().FillWidth(1.0f) // empty slot to eat the space for the alignment
+			+ SHorizontalBox::Slot()
 			.AutoWidth()
 			.HAlign(HAlign_Right)
 			.VAlign(VAlign_Center)
@@ -529,12 +535,12 @@ void SFoliageEditMeshDisplayItem::Construct( const FArguments& InArgs )
 			]
 		]
 
-		+SHorizontalBox::Slot()
+		+ SHorizontalBox::Slot()
 		.FillWidth(SPINBOX_SINGLE)
 		[
-			SNew( SHorizontalBox )
+			SNew(SHorizontalBox)
 
-			+SHorizontalBox::Slot()
+			+ SHorizontalBox::Slot()
 			.FillWidth(0.5f)
 			[
 				SNew(SSpinBox<float>)
@@ -545,7 +551,7 @@ void SFoliageEditMeshDisplayItem::Construct( const FArguments& InArgs )
 				.OnValueChanged(this, &SFoliageEditMeshDisplayItem::OnScaleYMinChanged)
 			]
 
-			+SHorizontalBox::Slot()
+			+ SHorizontalBox::Slot()
 			.AutoWidth()
 			.VAlign(VAlign_Center)
 			.Padding(StandardSidePadding)
@@ -554,18 +560,18 @@ void SFoliageEditMeshDisplayItem::Construct( const FArguments& InArgs )
 				.Text(LOCTEXT("Max", "Max"))
 			]
 
-			+SHorizontalBox::Slot()
+			+ SHorizontalBox::Slot()
 			.FillWidth(0.5f)
 			[
-				SNew(SSpinBox<float>)	
+				SNew(SSpinBox<float>)
 				.MinValue(0.0f)
 				.MaxValue(100.0f)
 				.MaxSliderValue(10.0f)
 				.Value(this, &SFoliageEditMeshDisplayItem::GetScaleYMax)
 				.OnValueChanged(this, &SFoliageEditMeshDisplayItem::OnScaleYMaxChanged)
 			]
-				
-			+SHorizontalBox::Slot()
+
+			+ SHorizontalBox::Slot()
 			.AutoWidth()
 			[
 				SNew(SCheckBox)
@@ -574,10 +580,11 @@ void SFoliageEditMeshDisplayItem::Construct( const FArguments& InArgs )
 			]
 		];
 
-	TSharedRef<SHorizontalBox> ScaleZBox = SNew( SHorizontalBox )
+	TSharedRef<SHorizontalBox> ScaleZBox =
+		SNew(SHorizontalBox)
 		.Visibility(this, &SFoliageEditMeshDisplayItem::IsNonUniformScalingVisible)
 
-		+SHorizontalBox::Slot()
+		+ SHorizontalBox::Slot()
 		.AutoWidth()
 		[
 			SNew(SCheckBox)
@@ -586,7 +593,7 @@ void SFoliageEditMeshDisplayItem::Construct( const FArguments& InArgs )
 			.IsChecked(this, &SFoliageEditMeshDisplayItem::IsScaleZReapplyChecked)
 		]
 
-		+SHorizontalBox::Slot()
+		+ SHorizontalBox::Slot()
 		.FillWidth(MAIN_TITLE)
 		.VAlign(VAlign_Center)
 		[
@@ -594,12 +601,12 @@ void SFoliageEditMeshDisplayItem::Construct( const FArguments& InArgs )
 			.Text(LOCTEXT("ScaleZ", "Scale Z"))
 		]
 
-		+SHorizontalBox::Slot()
+		+ SHorizontalBox::Slot()
 		.FillWidth(SPINBOX_PREFIX)
 		[
-			SNew( SHorizontalBox )
-			+SHorizontalBox::Slot().FillWidth(1.0f) // empty slot to eat the space for the alignment
-			+SHorizontalBox::Slot()
+			SNew(SHorizontalBox)
+			+ SHorizontalBox::Slot().FillWidth(1.0f) // empty slot to eat the space for the alignment
+			+ SHorizontalBox::Slot()
 			.AutoWidth()
 			.HAlign(HAlign_Right)
 			.VAlign(VAlign_Center)
@@ -610,12 +617,12 @@ void SFoliageEditMeshDisplayItem::Construct( const FArguments& InArgs )
 			]
 		]
 
-		+SHorizontalBox::Slot()
+		+ SHorizontalBox::Slot()
 		.FillWidth(SPINBOX_SINGLE)
 		[
-			SNew( SHorizontalBox )
-					
-			+SHorizontalBox::Slot()
+			SNew(SHorizontalBox)
+
+			+ SHorizontalBox::Slot()
 			.FillWidth(0.5f)
 			[
 				SNew(SSpinBox<float>)
@@ -626,7 +633,7 @@ void SFoliageEditMeshDisplayItem::Construct( const FArguments& InArgs )
 				.OnValueChanged(this, &SFoliageEditMeshDisplayItem::OnScaleZMinChanged)
 			]
 
-			+SHorizontalBox::Slot()
+			+ SHorizontalBox::Slot()
 			.AutoWidth()
 			.VAlign(VAlign_Center)
 			.Padding(StandardSidePadding)
@@ -635,10 +642,10 @@ void SFoliageEditMeshDisplayItem::Construct( const FArguments& InArgs )
 				.Text(LOCTEXT("Max", "Max"))
 			]
 
-			+SHorizontalBox::Slot()
+			+ SHorizontalBox::Slot()
 			.FillWidth(0.5f)
 			[
-				SNew(SSpinBox<float>)	
+				SNew(SSpinBox<float>)
 				.MinValue(0.0f)
 				.MaxValue(100.0f)
 				.MaxSliderValue(10.0f)
@@ -646,7 +653,7 @@ void SFoliageEditMeshDisplayItem::Construct( const FArguments& InArgs )
 				.OnValueChanged(this, &SFoliageEditMeshDisplayItem::OnScaleZMaxChanged)
 			]
 
-			+SHorizontalBox::Slot()
+			+ SHorizontalBox::Slot()
 			.AutoWidth()
 			[
 				SNew(SCheckBox)
@@ -655,9 +662,10 @@ void SFoliageEditMeshDisplayItem::Construct( const FArguments& InArgs )
 			]
 		];
 
-	TSharedRef<SHorizontalBox> ZOffsetBox = SNew( SHorizontalBox )
+	TSharedRef<SHorizontalBox> ZOffsetBox =
+		SNew(SHorizontalBox)
 
-		+SHorizontalBox::Slot()
+		+ SHorizontalBox::Slot()
 		.AutoWidth()
 		[
 			SNew(SCheckBox)
@@ -666,7 +674,7 @@ void SFoliageEditMeshDisplayItem::Construct( const FArguments& InArgs )
 			.IsChecked(this, &SFoliageEditMeshDisplayItem::IsZOffsetReapplyChecked)
 		]
 
-		+SHorizontalBox::Slot()
+		+ SHorizontalBox::Slot()
 		.FillWidth(MAIN_TITLE)
 		.VAlign(VAlign_Center)
 		[
@@ -674,12 +682,12 @@ void SFoliageEditMeshDisplayItem::Construct( const FArguments& InArgs )
 			.Text(LOCTEXT("ZOffset", "Z Offset"))
 		]
 
-		+SHorizontalBox::Slot()
+		+ SHorizontalBox::Slot()
 		.FillWidth(SPINBOX_PREFIX)
 		[
-			SNew( SHorizontalBox )
-			+SHorizontalBox::Slot().FillWidth(1.0f) // empty slot to eat the space for the alignment
-			+SHorizontalBox::Slot()
+			SNew(SHorizontalBox)
+			+ SHorizontalBox::Slot().FillWidth(1.0f) // empty slot to eat the space for the alignment
+			+ SHorizontalBox::Slot()
 			.AutoWidth()
 			.HAlign(HAlign_Right)
 			.VAlign(VAlign_Center)
@@ -690,12 +698,12 @@ void SFoliageEditMeshDisplayItem::Construct( const FArguments& InArgs )
 			]
 		]
 
-		+SHorizontalBox::Slot()
+		+ SHorizontalBox::Slot()
 		.FillWidth(SPINBOX_SINGLE)
 		[
-			SNew( SHorizontalBox )
+			SNew(SHorizontalBox)
 
-			+SHorizontalBox::Slot()
+			+ SHorizontalBox::Slot()
 			.FillWidth(0.5f)
 			[
 				SNew(SSpinBox<float>)
@@ -707,7 +715,7 @@ void SFoliageEditMeshDisplayItem::Construct( const FArguments& InArgs )
 				.OnValueChanged(this, &SFoliageEditMeshDisplayItem::OnZOffsetMin)
 			]
 
-			+SHorizontalBox::Slot()
+			+ SHorizontalBox::Slot()
 			.AutoWidth()
 			.VAlign(VAlign_Center)
 			.Padding(StandardSidePadding)
@@ -716,7 +724,7 @@ void SFoliageEditMeshDisplayItem::Construct( const FArguments& InArgs )
 				.Text(LOCTEXT("Max", "Max"))
 			]
 
-			+SHorizontalBox::Slot()
+			+ SHorizontalBox::Slot()
 			.FillWidth(0.5f)
 			[
 				SNew(SSpinBox<float>)
@@ -728,7 +736,7 @@ void SFoliageEditMeshDisplayItem::Construct( const FArguments& InArgs )
 				.OnValueChanged(this, &SFoliageEditMeshDisplayItem::OnZOffsetMax)
 			]
 
-			+SHorizontalBox::Slot()
+			+ SHorizontalBox::Slot()
 			.AutoWidth()
 			[
 				SNew(SCheckBox)
@@ -736,9 +744,10 @@ void SFoliageEditMeshDisplayItem::Construct( const FArguments& InArgs )
 			]
 		];
 
-	TSharedRef<SHorizontalBox> RandomPitchBox = SNew( SHorizontalBox )
+	TSharedRef<SHorizontalBox> RandomPitchBox =
+		SNew(SHorizontalBox)
 
-		+SHorizontalBox::Slot()
+		+ SHorizontalBox::Slot()
 		.AutoWidth()
 		[
 			SNew(SCheckBox)
@@ -747,7 +756,7 @@ void SFoliageEditMeshDisplayItem::Construct( const FArguments& InArgs )
 			.IsChecked(this, &SFoliageEditMeshDisplayItem::IsRandomPitchReapplyChecked)
 		]
 
-		+SHorizontalBox::Slot()
+		+ SHorizontalBox::Slot()
 		.FillWidth(MAIN_TITLE + SPINBOX_PREFIX)
 		.VAlign(VAlign_Center)
 		[
@@ -755,12 +764,12 @@ void SFoliageEditMeshDisplayItem::Construct( const FArguments& InArgs )
 			.Text(LOCTEXT("RandomPitchAngle", "Random Pitch +/-"))
 		]
 
-		+SHorizontalBox::Slot()
+		+ SHorizontalBox::Slot()
 		.FillWidth(SPINBOX_SINGLE)
 		[
-			SNew( SHorizontalBox )
+			SNew(SHorizontalBox)
 
-			+SHorizontalBox::Slot()
+			+ SHorizontalBox::Slot()
 			.FillWidth(1.f)
 			[
 				SNew(SSpinBox<float>)
@@ -771,7 +780,7 @@ void SFoliageEditMeshDisplayItem::Construct( const FArguments& InArgs )
 				.OnValueChanged(this, &SFoliageEditMeshDisplayItem::OnRandomPitchChanged)
 			]
 
-			+SHorizontalBox::Slot()
+			+ SHorizontalBox::Slot()
 			.AutoWidth()
 			[
 				SNew(SCheckBox)
@@ -779,9 +788,9 @@ void SFoliageEditMeshDisplayItem::Construct( const FArguments& InArgs )
 			]
 		];
 
-	TSharedRef<SHorizontalBox> GroundSlopeBox = SNew( SHorizontalBox )
+	TSharedRef<SHorizontalBox> GroundSlopeBox = SNew(SHorizontalBox)
 
-		+SHorizontalBox::Slot()
+		+ SHorizontalBox::Slot()
 		.AutoWidth()
 		[
 			SNew(SCheckBox)
@@ -790,7 +799,7 @@ void SFoliageEditMeshDisplayItem::Construct( const FArguments& InArgs )
 			.IsChecked(this, &SFoliageEditMeshDisplayItem::IsGroundSlopeReapplyChecked)
 		]
 
-		+SHorizontalBox::Slot()
+		+ SHorizontalBox::Slot()
 		.FillWidth(MAIN_TITLE + SPINBOX_PREFIX)
 		.VAlign(VAlign_Center)
 		[
@@ -798,12 +807,12 @@ void SFoliageEditMeshDisplayItem::Construct( const FArguments& InArgs )
 			.Text(LOCTEXT("GroundSlope", "Ground Slope"))
 		]
 
-		+SHorizontalBox::Slot()
+		+ SHorizontalBox::Slot()
 		.FillWidth(SPINBOX_SINGLE)
 		[
-			SNew( SHorizontalBox )
+			SNew(SHorizontalBox)
 
-			+SHorizontalBox::Slot()
+			+ SHorizontalBox::Slot()
 			.FillWidth(1.f)
 			[
 				SNew(SSpinBox<float>)
@@ -814,7 +823,7 @@ void SFoliageEditMeshDisplayItem::Construct( const FArguments& InArgs )
 				.OnValueChanged(this, &SFoliageEditMeshDisplayItem::OnGroundSlopeChanged)
 			]
 
-			+SHorizontalBox::Slot()
+			+ SHorizontalBox::Slot()
 			.AutoWidth()
 			[
 				SNew(SCheckBox)
@@ -822,9 +831,10 @@ void SFoliageEditMeshDisplayItem::Construct( const FArguments& InArgs )
 			]
 		];
 
-	TSharedRef<SHorizontalBox> HeightBox = SNew( SHorizontalBox )
+	TSharedRef<SHorizontalBox> HeightBox =
+		SNew(SHorizontalBox)
 
-		+SHorizontalBox::Slot()
+		+ SHorizontalBox::Slot()
 		.AutoWidth()
 		[
 			SNew(SCheckBox)
@@ -832,7 +842,7 @@ void SFoliageEditMeshDisplayItem::Construct( const FArguments& InArgs )
 			.OnCheckStateChanged(this, &SFoliageEditMeshDisplayItem::OnHeightReapply)
 			.IsChecked(this, &SFoliageEditMeshDisplayItem::IsHeightReapplyChecked)
 		]
-		+SHorizontalBox::Slot()
+		+ SHorizontalBox::Slot()
 		.FillWidth(MAIN_TITLE)
 		.VAlign(VAlign_Center)
 		[
@@ -840,12 +850,12 @@ void SFoliageEditMeshDisplayItem::Construct( const FArguments& InArgs )
 			.Text(LOCTEXT("Height", "Height"))
 		]
 
-		+SHorizontalBox::Slot()
+		+ SHorizontalBox::Slot()
 		.FillWidth(SPINBOX_PREFIX)
 		[
-			SNew( SHorizontalBox )
-			+SHorizontalBox::Slot().FillWidth(1.0f) // empty slot to eat the space for the alignment
-			+SHorizontalBox::Slot()
+			SNew(SHorizontalBox)
+			+ SHorizontalBox::Slot().FillWidth(1.0f) // empty slot to eat the space for the alignment
+			+ SHorizontalBox::Slot()
 			.AutoWidth()
 			.HAlign(HAlign_Right)
 			.VAlign(VAlign_Center)
@@ -856,12 +866,12 @@ void SFoliageEditMeshDisplayItem::Construct( const FArguments& InArgs )
 			]
 		]
 
-		+SHorizontalBox::Slot()
+		+ SHorizontalBox::Slot()
 		.FillWidth(SPINBOX_SINGLE)
 		[
-			SNew( SHorizontalBox )
+			SNew(SHorizontalBox)
 
-			+SHorizontalBox::Slot()
+			+ SHorizontalBox::Slot()
 			.FillWidth(0.5f)
 			[
 				SNew(SSpinBox<float>)
@@ -871,7 +881,7 @@ void SFoliageEditMeshDisplayItem::Construct( const FArguments& InArgs )
 				.OnValueChanged(this, &SFoliageEditMeshDisplayItem::OnHeightMinChanged)
 			]
 
-			+SHorizontalBox::Slot()
+			+ SHorizontalBox::Slot()
 			.AutoWidth()
 			.VAlign(VAlign_Center)
 			.Padding(StandardSidePadding)
@@ -880,17 +890,17 @@ void SFoliageEditMeshDisplayItem::Construct( const FArguments& InArgs )
 				.Text(LOCTEXT("Max", "Max"))
 			]
 
-			+SHorizontalBox::Slot()
+			+ SHorizontalBox::Slot()
 			.FillWidth(0.5f)
 			[
-				SNew(SSpinBox<float>)	
+				SNew(SSpinBox<float>)
 				.MinValue(-262144.0f)
 				.MaxValue(262144.0f)
 				.Value(this, &SFoliageEditMeshDisplayItem::GetHeightMax)
 				.OnValueChanged(this, &SFoliageEditMeshDisplayItem::OnHeightMaxChanged)
 			]
 
-			+SHorizontalBox::Slot()
+			+ SHorizontalBox::Slot()
 			.AutoWidth()
 			[
 				SNew(SCheckBox)
@@ -898,9 +908,10 @@ void SFoliageEditMeshDisplayItem::Construct( const FArguments& InArgs )
 			]
 		];
 
-	TSharedRef<SHorizontalBox> LandscapeLayerBox = SNew( SHorizontalBox )
+	TSharedRef<SHorizontalBox> LandscapeLayerBox =
+		SNew(SHorizontalBox)
 
-		+SHorizontalBox::Slot()
+		+ SHorizontalBox::Slot()
 		.AutoWidth()
 		[
 			SNew(SCheckBox)
@@ -909,7 +920,7 @@ void SFoliageEditMeshDisplayItem::Construct( const FArguments& InArgs )
 			.IsChecked(this, &SFoliageEditMeshDisplayItem::IsLandscapeLayerReapplyChecked)
 		]
 
-		+SHorizontalBox::Slot()
+		+ SHorizontalBox::Slot()
 		.FillWidth(MAIN_TITLE + SPINBOX_PREFIX)
 		.VAlign(VAlign_Center)
 		[
@@ -917,12 +928,12 @@ void SFoliageEditMeshDisplayItem::Construct( const FArguments& InArgs )
 			.Text(LOCTEXT("LandscapeLayer", "Landscape Layer"))
 		]
 
-		+SHorizontalBox::Slot()
+		+ SHorizontalBox::Slot()
 		.FillWidth(SPINBOX_SINGLE)
 		[
-			SNew( SHorizontalBox )
+			SNew(SHorizontalBox)
 
-			+SHorizontalBox::Slot()
+			+ SHorizontalBox::Slot()
 			.FillWidth(1.0f)
 			[
 				SNew(SEditableTextBox)
@@ -930,7 +941,7 @@ void SFoliageEditMeshDisplayItem::Construct( const FArguments& InArgs )
 				.OnTextChanged(this, &SFoliageEditMeshDisplayItem::OnLandscapeLayerChanged)
 			]
 
-			+SHorizontalBox::Slot()
+			+ SHorizontalBox::Slot()
 			.AutoWidth()
 			[
 				SNew(SCheckBox)
@@ -938,9 +949,10 @@ void SFoliageEditMeshDisplayItem::Construct( const FArguments& InArgs )
 			]
 		];
 
-	TSharedRef<SHorizontalBox> CollisionWithWorldBox = SNew( SHorizontalBox )
+	TSharedRef<SHorizontalBox> CollisionWithWorldBox =
+		SNew(SHorizontalBox)
 
-		+SHorizontalBox::Slot()
+		+ SHorizontalBox::Slot()
 		.AutoWidth()
 		[
 			SNew(SCheckBox)
@@ -949,7 +961,7 @@ void SFoliageEditMeshDisplayItem::Construct( const FArguments& InArgs )
 			.IsChecked(this, &SFoliageEditMeshDisplayItem::IsCollisionWithWorldReapplyChecked)
 		]
 
-		+SHorizontalBox::Slot()
+		+ SHorizontalBox::Slot()
 		.AutoWidth()
 		.VAlign(VAlign_Center)
 		[
@@ -962,18 +974,19 @@ void SFoliageEditMeshDisplayItem::Construct( const FArguments& InArgs )
 			]
 		];
 
-		TSharedRef<SHorizontalBox> CollisionScaleX = SNew( SHorizontalBox )
+	TSharedRef<SHorizontalBox> CollisionScaleX =
+		SNew(SHorizontalBox)
 		.Visibility(this, &SFoliageEditMeshDisplayItem::IsCollisionWithWorldVisible)
 
 		// Dummy Checkbox
-		+SHorizontalBox::Slot()
+		+ SHorizontalBox::Slot()
 		.AutoWidth()
 		[
 			SNew(SCheckBox)
 			.Visibility(this, &SFoliageEditMeshDisplayItem::IsReapplySettingsVisible_Dummy)
 		]
 
-		+SHorizontalBox::Slot()
+		+ SHorizontalBox::Slot()
 		.FillWidth(MAIN_TITLE)
 		.VAlign(VAlign_Center)
 		[
@@ -981,12 +994,12 @@ void SFoliageEditMeshDisplayItem::Construct( const FArguments& InArgs )
 			.Text(LOCTEXT("CollisionScale", "Collision Scale"))
 		]
 
-		+SHorizontalBox::Slot()
+		+ SHorizontalBox::Slot()
 		.FillWidth(SPINBOX_PREFIX)
 		[
-			SNew( SHorizontalBox )
-			+SHorizontalBox::Slot().FillWidth(1.0f) // empty slot to eat the space for the alignment
-			+SHorizontalBox::Slot()
+			SNew(SHorizontalBox)
+			+ SHorizontalBox::Slot().FillWidth(1.0f) // empty slot to eat the space for the alignment
+			+ SHorizontalBox::Slot()
 			.AutoWidth()
 			.HAlign(HAlign_Right)
 			.VAlign(VAlign_Center)
@@ -997,15 +1010,15 @@ void SFoliageEditMeshDisplayItem::Construct( const FArguments& InArgs )
 			]
 		]
 
-		+SHorizontalBox::Slot()
+		+ SHorizontalBox::Slot()
 		.FillWidth(SPINBOX_SINGLE)
 		[
-			SNew( SHorizontalBox )
+			SNew(SHorizontalBox)
 
-			+SHorizontalBox::Slot()
+			+ SHorizontalBox::Slot()
 			.FillWidth(0.333333f)
 			[
-				SNew(SSpinBox<float>)	
+				SNew(SSpinBox<float>)
 				.MinValue(0.01f)
 				.MaxValue(5.0f)
 				.MaxSliderValue(1.0f)
@@ -1014,7 +1027,7 @@ void SFoliageEditMeshDisplayItem::Construct( const FArguments& InArgs )
 			]
 
 			// Y
-			+SHorizontalBox::Slot()
+			+ SHorizontalBox::Slot()
 			.AutoWidth()
 			.VAlign(VAlign_Center)
 			.Padding(StandardSidePadding)
@@ -1023,10 +1036,10 @@ void SFoliageEditMeshDisplayItem::Construct( const FArguments& InArgs )
 				.Text(LOCTEXT("Y", "Y"))
 			]
 
-			+SHorizontalBox::Slot()
+			+ SHorizontalBox::Slot()
 			.FillWidth(0.333333f)
 			[
-				SNew(SSpinBox<float>)	
+				SNew(SSpinBox<float>)
 				.MinValue(0.01f)
 				.MaxValue(5.0f)
 				.MaxSliderValue(1.0f)
@@ -1035,7 +1048,7 @@ void SFoliageEditMeshDisplayItem::Construct( const FArguments& InArgs )
 			]
 
 			// Z
-			+SHorizontalBox::Slot()
+			+ SHorizontalBox::Slot()
 			.AutoWidth()
 			.VAlign(VAlign_Center)
 			.Padding(StandardSidePadding)
@@ -1044,10 +1057,10 @@ void SFoliageEditMeshDisplayItem::Construct( const FArguments& InArgs )
 				.Text(LOCTEXT("Z", "Z"))
 			]
 
-			+SHorizontalBox::Slot()
+			+ SHorizontalBox::Slot()
 			.FillWidth(0.333333f)
 			[
-				SNew(SSpinBox<float>)	
+				SNew(SSpinBox<float>)
 				.MinValue(0.01f)
 				.MaxValue(5.0f)
 				.MaxSliderValue(1.0f)
@@ -1055,7 +1068,7 @@ void SFoliageEditMeshDisplayItem::Construct( const FArguments& InArgs )
 				.OnValueChanged(this, &SFoliageEditMeshDisplayItem::OnCollisionScaleZChanged)
 			]
 
-			+SHorizontalBox::Slot()
+			+ SHorizontalBox::Slot()
 			.AutoWidth()
 			[
 				SNew(SCheckBox)
@@ -1063,9 +1076,10 @@ void SFoliageEditMeshDisplayItem::Construct( const FArguments& InArgs )
 			]
 		];
 
-	TSharedRef<SHorizontalBox> VertexColorMaskBox = SNew( SHorizontalBox )
+	TSharedRef<SHorizontalBox> VertexColorMaskBox =
+		SNew(SHorizontalBox)
 
-		+SHorizontalBox::Slot()
+		+ SHorizontalBox::Slot()
 		.AutoWidth()
 		[
 			SNew(SCheckBox)
@@ -1074,7 +1088,7 @@ void SFoliageEditMeshDisplayItem::Construct( const FArguments& InArgs )
 			.IsChecked(this, &SFoliageEditMeshDisplayItem::IsVertexColorMaskReapplyChecked)
 		]
 
-		+SHorizontalBox::Slot()
+		+ SHorizontalBox::Slot()
 		.FillWidth(MAIN_TITLE + SPINBOX_PREFIX)
 		.VAlign(VAlign_Center)
 		[
@@ -1082,12 +1096,12 @@ void SFoliageEditMeshDisplayItem::Construct( const FArguments& InArgs )
 			.Text(LOCTEXT("VertexColorMask", "Vertex Color Mask"))
 		]
 
-		+SHorizontalBox::Slot()
+		+ SHorizontalBox::Slot()
 		.FillWidth(SPINBOX_SINGLE)
 		[
-			SNew( SHorizontalBox )
+			SNew(SHorizontalBox)
 
-			+SHorizontalBox::Slot()
+			+ SHorizontalBox::Slot()
 			.AutoWidth()
 			.Padding(StandardSidePadding)
 			.VAlign(VAlign_Center)
@@ -1101,7 +1115,7 @@ void SFoliageEditMeshDisplayItem::Construct( const FArguments& InArgs )
 				]
 			]
 
-			+SHorizontalBox::Slot()
+			+ SHorizontalBox::Slot()
 			.AutoWidth()
 			.Padding(StandardSidePadding)
 			.VAlign(VAlign_Center)
@@ -1115,7 +1129,7 @@ void SFoliageEditMeshDisplayItem::Construct( const FArguments& InArgs )
 				]
 			]
 
-			+SHorizontalBox::Slot()
+			+ SHorizontalBox::Slot()
 			.AutoWidth()
 			.Padding(StandardSidePadding)
 			.VAlign(VAlign_Center)
@@ -1129,7 +1143,7 @@ void SFoliageEditMeshDisplayItem::Construct( const FArguments& InArgs )
 				]
 			]
 
-			+SHorizontalBox::Slot()
+			+ SHorizontalBox::Slot()
 			.AutoWidth()
 			.Padding(StandardSidePadding)
 			.VAlign(VAlign_Center)
@@ -1145,7 +1159,7 @@ void SFoliageEditMeshDisplayItem::Construct( const FArguments& InArgs )
 			]
 
 
-			+SHorizontalBox::Slot()
+			+ SHorizontalBox::Slot()
 			.AutoWidth()
 			.Padding(StandardSidePadding)
 			.VAlign(VAlign_Center)
@@ -1161,18 +1175,19 @@ void SFoliageEditMeshDisplayItem::Construct( const FArguments& InArgs )
 			]
 		];
 
-	TSharedRef<SHorizontalBox> VertexColorMaskThresholdBox = SNew( SHorizontalBox )
+	TSharedRef<SHorizontalBox> VertexColorMaskThresholdBox =
+		SNew(SHorizontalBox)
 		.Visibility(this, &SFoliageEditMeshDisplayItem::IsVertexColorMaskThresholdVisible)
 
 		// Dummy Checkbox
-		+SHorizontalBox::Slot()
+		+ SHorizontalBox::Slot()
 		.AutoWidth()
 		[
 			SNew(SCheckBox)
 			.Visibility(this, &SFoliageEditMeshDisplayItem::IsReapplySettingsVisible_Dummy)
 		]
 
-		+SHorizontalBox::Slot()
+		+ SHorizontalBox::Slot()
 		.FillWidth(MAIN_TITLE + SPINBOX_PREFIX)
 		.VAlign(VAlign_Center)
 		[
@@ -1181,12 +1196,12 @@ void SFoliageEditMeshDisplayItem::Construct( const FArguments& InArgs )
 			.Visibility(this, &SFoliageEditMeshDisplayItem::IsVertexColorMaskThresholdVisible)
 		]
 
-		+SHorizontalBox::Slot()
+		+ SHorizontalBox::Slot()
 		.FillWidth(SPINBOX_SINGLE)
 		[
-			SNew( SHorizontalBox )
+			SNew(SHorizontalBox)
 
-			+SHorizontalBox::Slot()
+			+ SHorizontalBox::Slot()
 			.AutoWidth()
 			.Padding(StandardSidePadding)
 			.VAlign(VAlign_Center)
@@ -1197,18 +1212,18 @@ void SFoliageEditMeshDisplayItem::Construct( const FArguments& InArgs )
 				.Value(this, &SFoliageEditMeshDisplayItem::GetVertexColorMaskThreshold)
 				.OnValueChanged(this, &SFoliageEditMeshDisplayItem::OnVertexColorMaskThresholdChanged)
 			]
-				
-			+SHorizontalBox::Slot()
+
+			+ SHorizontalBox::Slot()
 			.AutoWidth()
 			.Padding(StandardSidePadding)
 			.VAlign(VAlign_Center)
 			[
-				SNew(SCheckBox)									
+				SNew(SCheckBox)
 				.OnCheckStateChanged(this, &SFoliageEditMeshDisplayItem::OnVertexColorMaskInvert)
 				.IsChecked(this, &SFoliageEditMeshDisplayItem::IsVertexColorMaskInvertChecked)
 			]
 
-			+SHorizontalBox::Slot()
+			+ SHorizontalBox::Slot()
 			.AutoWidth()
 			.VAlign(VAlign_Center)
 			[
@@ -1217,25 +1232,26 @@ void SFoliageEditMeshDisplayItem::Construct( const FArguments& InArgs )
 			]
 		];
 
-	TSharedRef<SVerticalBox> PaintSettings = SNew(SVerticalBox)
+	TSharedRef<SVerticalBox> PaintSettings =
+		SNew(SVerticalBox)
 		.Visibility(this, &SFoliageEditMeshDisplayItem::IsPaintSettingsVisible)
 
-		+SVerticalBox::Slot()
+		+ SVerticalBox::Slot()
 		.AutoHeight()
 		[
 			DensityBox
 		]
 
 		// Align to Normal Checkbox
-		+SVerticalBox::Slot()
-			.AutoHeight()
-			.Padding(StandardPadding)
-			[
-				AlignToNormalBox
-			]
+		+ SVerticalBox::Slot()
+		.AutoHeight()
+		.Padding(StandardPadding)
+		[
+			AlignToNormalBox
+		]
 
 		// Max Angle
-		+SVerticalBox::Slot()
+		+ SVerticalBox::Slot()
 		.AutoHeight()
 		.Padding(StandardPadding)
 		[
@@ -1243,22 +1259,22 @@ void SFoliageEditMeshDisplayItem::Construct( const FArguments& InArgs )
 		]
 
 		// Random Yaw Checkbox
-		+SVerticalBox::Slot()
-			.AutoHeight()
-			.Padding(StandardPadding)
-			[
-				RandomYawBox
-			]
+		+ SVerticalBox::Slot()
+		.AutoHeight()
+		.Padding(StandardPadding)
+		[
+			RandomYawBox
+		]
 
-		+SVerticalBox::Slot()
-			.AutoHeight()
-			.Padding(StandardPadding)
-			[
-				UniformScaleBox
-			]
+		+ SVerticalBox::Slot()
+		.AutoHeight()
+		.Padding(StandardPadding)
+		[
+			UniformScaleBox
+		]
 
 		// Uniform Scale
-		+SVerticalBox::Slot()
+		+ SVerticalBox::Slot()
 		.AutoHeight()
 		.Padding(StandardPadding)
 		[
@@ -1266,7 +1282,7 @@ void SFoliageEditMeshDisplayItem::Construct( const FArguments& InArgs )
 		]
 
 		// Scale X
-		+SVerticalBox::Slot()
+		+ SVerticalBox::Slot()
 		.AutoHeight()
 		.Padding(StandardPadding)
 		[
@@ -1274,7 +1290,7 @@ void SFoliageEditMeshDisplayItem::Construct( const FArguments& InArgs )
 		]
 
 		// Scale Y
-		+SVerticalBox::Slot()
+		+ SVerticalBox::Slot()
 		.AutoHeight()
 		.Padding(StandardPadding)
 		[
@@ -1282,7 +1298,7 @@ void SFoliageEditMeshDisplayItem::Construct( const FArguments& InArgs )
 		]
 
 		// Scale Z
-		+SVerticalBox::Slot()
+		+ SVerticalBox::Slot()
 		.AutoHeight()
 		.Padding(StandardPadding)
 		[
@@ -1290,7 +1306,7 @@ void SFoliageEditMeshDisplayItem::Construct( const FArguments& InArgs )
 		]
 
 		// Z Offset
-		+SVerticalBox::Slot()
+		+ SVerticalBox::Slot()
 		.AutoHeight()
 		.Padding(StandardPadding)
 		[
@@ -1298,7 +1314,7 @@ void SFoliageEditMeshDisplayItem::Construct( const FArguments& InArgs )
 		]
 
 		// Random Pitch +/-
-		+SVerticalBox::Slot()
+		+ SVerticalBox::Slot()
 		.AutoHeight()
 		.Padding(StandardPadding)
 		[
@@ -1306,7 +1322,7 @@ void SFoliageEditMeshDisplayItem::Construct( const FArguments& InArgs )
 		]
 
 		// Ground Slope
-		+SVerticalBox::Slot()
+		+ SVerticalBox::Slot()
 		.AutoHeight()
 		.Padding(StandardPadding)
 		[
@@ -1314,7 +1330,7 @@ void SFoliageEditMeshDisplayItem::Construct( const FArguments& InArgs )
 		]
 
 		// Height
-		+SVerticalBox::Slot()
+		+ SVerticalBox::Slot()
 		.AutoHeight()
 		.Padding(StandardPadding)
 		[
@@ -1322,14 +1338,14 @@ void SFoliageEditMeshDisplayItem::Construct( const FArguments& InArgs )
 		]
 
 		// Landscape Layer
-		+SVerticalBox::Slot()
+		+ SVerticalBox::Slot()
 		.AutoHeight()
 		.Padding(StandardPadding)
 		[
 			LandscapeLayerBox
 		]
 
-		+SVerticalBox::Slot()
+		+ SVerticalBox::Slot()
 		.AutoHeight()
 		.Padding(StandardPadding)
 		[
@@ -1337,7 +1353,7 @@ void SFoliageEditMeshDisplayItem::Construct( const FArguments& InArgs )
 		]
 
 		// Collision ScaleX
-		+SVerticalBox::Slot()
+		+ SVerticalBox::Slot()
 		.AutoHeight()
 		.Padding(StandardPadding)
 		[
@@ -1345,14 +1361,14 @@ void SFoliageEditMeshDisplayItem::Construct( const FArguments& InArgs )
 		]
 
 		// Vertex Color Mask
-		+SVerticalBox::Slot()
+		+ SVerticalBox::Slot()
 		.AutoHeight()
 		.Padding(StandardPadding)
 		[
 			VertexColorMaskBox
 		]
 
-		+SVerticalBox::Slot()
+		+ SVerticalBox::Slot()
 		.AutoHeight()
 		.Padding(StandardPadding)
 		[
@@ -1360,7 +1376,7 @@ void SFoliageEditMeshDisplayItem::Construct( const FArguments& InArgs )
 		];
 
 	this->ChildSlot
-	.Padding( 6.0f )
+	.Padding(6.0f)
 	.VAlign(VAlign_Top)
 	[
 		SNew(SBorder)
@@ -1368,8 +1384,8 @@ void SFoliageEditMeshDisplayItem::Construct( const FArguments& InArgs )
 		.Padding(0.0f)
 		[
 			SNew(SHorizontalBox)
-	
-			+SHorizontalBox::Slot()
+
+			+ SHorizontalBox::Slot()
 			.AutoWidth()
 			[
 				SNew(SBorder)
@@ -1391,27 +1407,35 @@ void SFoliageEditMeshDisplayItem::Construct( const FArguments& InArgs )
 				]
 			]
 
-			+SHorizontalBox::Slot()
-				.FillWidth(1.0f)
-				.Padding(2.0f)
+			+ SHorizontalBox::Slot()
+			.FillWidth(1.0f)
+			.Padding(2.0f)
 			[
-				SNew( SVerticalBox )
+				SNew(SVerticalBox)
 
-				+SVerticalBox::Slot()
+				+ SVerticalBox::Slot()
 				.AutoHeight()
 				.Padding(StandardPadding)
 				[
-					SNew( SHorizontalBox )
+					SNew(SHorizontalBox)
 					.Visibility(this, &SFoliageEditMeshDisplayItem::IsSettingsVisible)
 
-					+SHorizontalBox::Slot()
+					+ SHorizontalBox::Slot()
 					.AutoWidth()
 					.VAlign(VAlign_Center)
 					[
 						Toolbar.MakeWidget()
 					]
+				]
 
-					+SHorizontalBox::Slot()
+				+ SVerticalBox::Slot()
+				.AutoHeight()
+				.Padding(StandardPadding)
+				[
+					SNew(SHorizontalBox)
+					.Visibility(this, &SFoliageEditMeshDisplayItem::IsSettingsVisible)
+
+					+ SHorizontalBox::Slot()
 					.FillWidth(1.0f)
 					.HAlign(HAlign_Center)
 					.VAlign(VAlign_Center)
@@ -1420,57 +1444,57 @@ void SFoliageEditMeshDisplayItem::Construct( const FArguments& InArgs )
 						.Text(this, &SFoliageEditMeshDisplayItem::GetStaticMeshname)
 					]
 
-					+SHorizontalBox::Slot()
+					+ SHorizontalBox::Slot()
 					.AutoWidth()
 					.VAlign(VAlign_Center)
 					[
 						SNew(SButton)
-						.ButtonStyle( FEditorStyle::Get(), "NoBorder" )
-						.OnClicked( this, &SFoliageEditMeshDisplayItem::OnReplace )
+						.ButtonStyle(FEditorStyle::Get(), "NoBorder")
+						.OnClicked(this, &SFoliageEditMeshDisplayItem::OnReplace)
 						.ToolTipText(NSLOCTEXT("FoliageEdMode", "Replace_Tooltip", "Replace all instances with the Static Mesh currently selected in the Content Browser."))
 						[
 							SNew(SImage)
-							.Image( FEditorStyle::GetBrush(TEXT("ContentReference.UseSelectionFromContentBrowser")) )
+							.Image(FEditorStyle::GetBrush(TEXT("ContentReference.UseSelectionFromContentBrowser")))
 						]
 					]
 
-					+SHorizontalBox::Slot()
+					+ SHorizontalBox::Slot()
 					.AutoWidth()
 					.VAlign(VAlign_Center)
 					[
-						SNew(SButton)	
-						.ButtonStyle( FEditorStyle::Get(), "NoBorder" )
-						.OnClicked( this, &SFoliageEditMeshDisplayItem::OnSync )
+						SNew(SButton)
+						.ButtonStyle(FEditorStyle::Get(), "NoBorder")
+						.OnClicked(this, &SFoliageEditMeshDisplayItem::OnSync)
 						.ToolTipText(NSLOCTEXT("FoliageEdMode", "FindInContentBrowser_Tooltip", "Find this Static Mesh in the Content Browser."))
 						[
 							SNew(SImage)
-							.Image( FEditorStyle::GetBrush(TEXT("ContentReference.FindInContentBrowser")) )
+							.Image(FEditorStyle::GetBrush(TEXT("ContentReference.FindInContentBrowser")))
 						]
 					]
 
-					+SHorizontalBox::Slot()
+					+ SHorizontalBox::Slot()
 					.AutoWidth()
 					.VAlign(VAlign_Center)
 					[
-						SNew(SButton)	
-						.ButtonStyle( FEditorStyle::Get(), "NoBorder" )
-						.OnClicked( this, &SFoliageEditMeshDisplayItem::OnRemove )
+						SNew(SButton)
+						.ButtonStyle(FEditorStyle::Get(), "NoBorder")
+						.OnClicked(this, &SFoliageEditMeshDisplayItem::OnRemove)
 						.ToolTipText(NSLOCTEXT("FoliageEdMode", "Remove_Tooltip", "Delete all foliage instances of this Static Mesh."))
 						[
 							SNew(SImage)
-							.Image( FEditorStyle::GetBrush(TEXT("ContentReference.Clear")) )
+							.Image(FEditorStyle::GetBrush(TEXT("ContentReference.Clear")))
 						]
 					]
 				]
 
-				+SVerticalBox::Slot()
+				+ SVerticalBox::Slot()
 				.AutoHeight()
 				.Padding(StandardPadding)
 				[
 					SNew(SHorizontalBox)
 					.Visibility(this, &SFoliageEditMeshDisplayItem::IsSettingsVisible)
 
-					+SHorizontalBox::Slot()
+					+ SHorizontalBox::Slot()
 					.FillWidth(1.0f)
 					.HAlign(HAlign_Center)
 					.VAlign(VAlign_Center)
@@ -1479,52 +1503,52 @@ void SFoliageEditMeshDisplayItem::Construct( const FArguments& InArgs )
 						.Text(this, &SFoliageEditMeshDisplayItem::GetSettingsLabelText)
 					]
 
-					+SHorizontalBox::Slot()
+					+ SHorizontalBox::Slot()
 					.AutoWidth()
 					[
 						SNew(SButton)
-						.ButtonStyle( FEditorStyle::Get(), "NoBorder" )
-						.OnClicked( this, &SFoliageEditMeshDisplayItem::OnOpenSettings )
+						.ButtonStyle(FEditorStyle::Get(), "NoBorder")
+						.OnClicked(this, &SFoliageEditMeshDisplayItem::OnOpenSettings)
 						.ToolTipText(NSLOCTEXT("FoliageEdMode", "OpenSettings_Tooltip", "Use the InstancedFoliageSettings currently selected in the Content Browser."))
 						[
 							SNew(SImage)
-							.Image( FEditorStyle::GetBrush(TEXT("FoliageEditMode.OpenSettings")) )
+							.Image(FEditorStyle::GetBrush(TEXT("FoliageEditMode.OpenSettings")))
 						]
 					]
-				
-					+SHorizontalBox::Slot()
+
+					+ SHorizontalBox::Slot()
 					.AutoWidth()
 					[
 						SNew(SButton)
-						.ButtonStyle( FEditorStyle::Get(), "NoBorder" )
-						.OnClicked( this, &SFoliageEditMeshDisplayItem::OnSaveRemoveSettings )
+						.ButtonStyle(FEditorStyle::Get(), "NoBorder")
+						.OnClicked(this, &SFoliageEditMeshDisplayItem::OnSaveRemoveSettings)
 						.ToolTipText(this, &SFoliageEditMeshDisplayItem::GetSaveRemoveSettingsTooltip)
 						[
 							SNew(SImage)
-							.Image( this, &SFoliageEditMeshDisplayItem::GetSaveSettingsBrush )
+							.Image(this, &SFoliageEditMeshDisplayItem::GetSaveSettingsBrush)
 						]
 					]
 				]
 
-				+SVerticalBox::Slot()
+				+ SVerticalBox::Slot()
 				.AutoHeight()
 				[
 					SNew(SHorizontalBox)
 					.Visibility(this, &SFoliageEditMeshDisplayItem::IsNoSettingsVisible)
 
-					+SHorizontalBox::Slot()
+					+ SHorizontalBox::Slot()
 					.FillWidth(1.0f)
 					[
 						SNew(SVerticalBox)
 
-						+SVerticalBox::Slot()
+						+ SVerticalBox::Slot()
 						.AutoHeight()
 						.HAlign(HAlign_Left)
 						[
-							Toolbar.MakeWidget()																	
+							Toolbar.MakeWidget()
 						]
 
-						+SVerticalBox::Slot()
+						+ SVerticalBox::Slot()
 						.FillHeight(1.0f)
 						.HAlign(HAlign_Center)
 						.VAlign(VAlign_Center)
@@ -1534,62 +1558,62 @@ void SFoliageEditMeshDisplayItem::Construct( const FArguments& InArgs )
 						]
 					]
 
-					+SHorizontalBox::Slot()
+					+ SHorizontalBox::Slot()
 					.AutoWidth()
 					[
 						ThumbnailBox.ToSharedRef()
 					]
 				]
 
-				+SVerticalBox::Slot()
+				+ SVerticalBox::Slot()
 				.AutoHeight()
 				[
 					SNew(SVerticalBox)
 					.Visibility(this, &SFoliageEditMeshDisplayItem::IsClusterSettingsVisible)
 
-					+SVerticalBox::Slot()
+					+ SVerticalBox::Slot()
 					.AutoHeight()
 					[
 						SNew(SHorizontalBox)
 
-						+SHorizontalBox::Slot()
+						+ SHorizontalBox::Slot()
 						.FillWidth(1.0f)
 						[
 							SNew(SVerticalBox)
 
-							+SVerticalBox::Slot()
+							+ SVerticalBox::Slot()
 							.AutoHeight()
 							.VAlign(VAlign_Center)
 							[
 								SNew(STextBlock)
-								.Text( this, &SFoliageEditMeshDisplayItem::GetInstanceCountString)
+								.Text(this, &SFoliageEditMeshDisplayItem::GetInstanceCountString)
 							]
 
-							+SVerticalBox::Slot()
+							+ SVerticalBox::Slot()
 							.AutoHeight()
 							.VAlign(VAlign_Center)
 							[
 								SNew(STextBlock)
-								.Text( this, &SFoliageEditMeshDisplayItem::GetInstanceClusterCountString)
+								.Text(this, &SFoliageEditMeshDisplayItem::GetInstanceClusterCountString)
 							]
 						]
 
-						+SHorizontalBox::Slot()
+						+ SHorizontalBox::Slot()
 						.AutoWidth()
 						.Padding(StandardPadding)
 						[
-							ThumbnailBox.ToSharedRef()		
+							ThumbnailBox.ToSharedRef()
 						]
 					]
 
-					+SVerticalBox::Slot()
+					+ SVerticalBox::Slot()
 					.AutoHeight()
 					[
 						ClusterSettingsDetails.ToSharedRef()
 					]
 				]
 
-				+SVerticalBox::Slot()
+				+ SVerticalBox::Slot()
 				.AutoHeight()
 				[
 					PaintSettings
@@ -1600,9 +1624,9 @@ void SFoliageEditMeshDisplayItem::Construct( const FArguments& InArgs )
 }
 END_SLATE_FUNCTION_BUILD_OPTIMIZATION
 
-void SFoliageEditMeshDisplayItem::NotifyPostChange( const FPropertyChangedEvent& PropertyChangedEvent, UProperty* PropertyThatChanged )
+void SFoliageEditMeshDisplayItem::NotifyPostChange(const FPropertyChangedEvent& PropertyChangedEvent, UProperty* PropertyThatChanged)
 {
-	FoliageEditPtr.Pin()->GetFoliageEditMode()->ReallocateClusters(AssociatedStaticMesh);
+	FoliageEditPtr.Pin()->GetFoliageEditMode()->ReallocateClusters(FoliageSettingsPtr);
 }
 
 void SFoliageEditMeshDisplayItem::BindCommands()
@@ -1611,21 +1635,21 @@ void SFoliageEditMeshDisplayItem::BindCommands()
 
 	UICommandList->MapAction(
 		Commands.SetNoSettings,
-		FExecuteAction::CreateSP( this, &SFoliageEditMeshDisplayItem::OnNoSettings ),
+		FExecuteAction::CreateSP(this, &SFoliageEditMeshDisplayItem::OnNoSettings),
 		FCanExecuteAction(),
-		FIsActionChecked::CreateSP( this, &SFoliageEditMeshDisplayItem::IsNoSettingsEnabled ) );
+		FIsActionChecked::CreateSP(this, &SFoliageEditMeshDisplayItem::IsNoSettingsEnabled));
 
 	UICommandList->MapAction(
 		Commands.SetPaintSettings,
-		FExecuteAction::CreateSP( this, &SFoliageEditMeshDisplayItem::OnPaintSettings ),
+		FExecuteAction::CreateSP(this, &SFoliageEditMeshDisplayItem::OnPaintSettings),
 		FCanExecuteAction(),
-		FIsActionChecked::CreateSP( this, &SFoliageEditMeshDisplayItem::IsPaintSettingsEnabled ) );
+		FIsActionChecked::CreateSP(this, &SFoliageEditMeshDisplayItem::IsPaintSettingsEnabled));
 
 	UICommandList->MapAction(
 		Commands.SetClusterSettings,
-		FExecuteAction::CreateSP( this, &SFoliageEditMeshDisplayItem::OnClusterSettings ),
+		FExecuteAction::CreateSP(this, &SFoliageEditMeshDisplayItem::OnClusterSettings),
 		FCanExecuteAction(),
-		FIsActionChecked::CreateSP( this, &SFoliageEditMeshDisplayItem::IsClusterSettingsEnabled ) );
+		FIsActionChecked::CreateSP(this, &SFoliageEditMeshDisplayItem::IsClusterSettingsEnabled));
 }
 
 ECurrentViewSettings::Type SFoliageEditMeshDisplayItem::GetCurrentDisplayStatus() const
@@ -1638,13 +1662,12 @@ void SFoliageEditMeshDisplayItem::SetCurrentDisplayStatus(ECurrentViewSettings::
 	CurrentViewSettings = InDisplayStatus;
 }
 
-void SFoliageEditMeshDisplayItem::Replace(UInstancedFoliageSettings* InFoliageSettingsPtr, UStaticMesh* InAssociatedStaticMesh, TSharedPtr<FAssetThumbnail> InAssetThumbnail, TSharedPtr<FFoliageMeshUIInfo> InFoliageMeshUIInfo)
+void SFoliageEditMeshDisplayItem::Replace(UFoliageType* InFoliageSettingsPtr, TSharedPtr<FAssetThumbnail> InAssetThumbnail, TSharedPtr<FFoliageMeshUIInfo> InFoliageMeshUIInfo)
 {
 	FoliageSettingsPtr = InFoliageSettingsPtr;
-	AssociatedStaticMesh = InAssociatedStaticMesh;
 	FoliageMeshUIInfo = InFoliageMeshUIInfo;
-	Thumbnail = InAssetThumbnail;
-	ThumbnailWidget = Thumbnail->MakeThumbnailWidget();
+	Thumbnail = MoveTemp(InAssetThumbnail);
+	ThumbnailWidget = Thumbnail.Get()->MakeThumbnailWidget();
 
 	ThumbnailWidgetBorder->SetContent(ThumbnailWidget.ToSharedRef());
 }
@@ -1679,73 +1702,69 @@ bool SFoliageEditMeshDisplayItem::IsClusterSettingsEnabled() const
 	return CurrentViewSettings == ECurrentViewSettings::ShowClusterSettings;
 }
 
-FString SFoliageEditMeshDisplayItem::GetStaticMeshname() const
+FText SFoliageEditMeshDisplayItem::GetStaticMeshname() const
 {
-	return AssociatedStaticMesh->GetName();
+	return FText::FromName(FoliageSettingsPtr->GetStaticMesh()->GetFName());
 }
 
-FString SFoliageEditMeshDisplayItem::GetSettingsLabelText() const
+FText SFoliageEditMeshDisplayItem::GetSettingsLabelText() const
 {
-	return FoliageSettingsPtr->GetOuter()->IsA(UPackage::StaticClass()) ? *FoliageSettingsPtr->GetPathName() : LOCTEXT("NotUsingSharedSettings", "(not using shared settings)").ToString();
+	return FoliageSettingsPtr->GetOuter()->IsA(UPackage::StaticClass()) ? FText::FromString(FoliageSettingsPtr->GetPathName()) : LOCTEXT("NotUsingSharedSettings", "(not using shared settings)");
 }
 
 EVisibility  SFoliageEditMeshDisplayItem::IsNoSettingsVisible() const
 {
-	return CurrentViewSettings == ECurrentViewSettings::ShowNone? EVisibility::Visible : EVisibility::Collapsed;
+	return CurrentViewSettings == ECurrentViewSettings::ShowNone ? EVisibility::Visible : EVisibility::Collapsed;
 }
 
 EVisibility  SFoliageEditMeshDisplayItem::IsSettingsVisible() const
 {
-	return CurrentViewSettings == ECurrentViewSettings::ShowNone? EVisibility::Collapsed : EVisibility::Visible;
+	return CurrentViewSettings == ECurrentViewSettings::ShowNone ? EVisibility::Collapsed : EVisibility::Visible;
 }
 
 EVisibility SFoliageEditMeshDisplayItem::IsPaintSettingsVisible() const
 {
-	return CurrentViewSettings == ECurrentViewSettings::ShowPaintSettings? EVisibility::Visible : EVisibility::Collapsed;
+	return CurrentViewSettings == ECurrentViewSettings::ShowPaintSettings ? EVisibility::Visible : EVisibility::Collapsed;
 }
 
 EVisibility SFoliageEditMeshDisplayItem::IsClusterSettingsVisible() const
 {
-	return CurrentViewSettings == ECurrentViewSettings::ShowClusterSettings? EVisibility::Visible : EVisibility::Collapsed;
+	return CurrentViewSettings == ECurrentViewSettings::ShowClusterSettings ? EVisibility::Visible : EVisibility::Collapsed;
 }
 
 FReply SFoliageEditMeshDisplayItem::OnReplace()
 {
-	if(AssociatedStaticMesh)
+	// Get current selection from content browser
+	FEditorDelegates::LoadSelectedAssetsIfNeeded.Broadcast();
+	USelection* SelectedSet = GEditor->GetSelectedSet(UStaticMesh::StaticClass());
+	UStaticMesh* SelectedStaticMesh = Cast<UStaticMesh>(SelectedSet->GetTop(UStaticMesh::StaticClass()));
+	if (SelectedStaticMesh != NULL)
 	{
-		// Get current selection from content browser
-		FEditorDelegates::LoadSelectedAssetsIfNeeded.Broadcast();
-		USelection* SelectedSet = GEditor->GetSelectedSet( UStaticMesh::StaticClass() );
-		UStaticMesh* SelectedStaticMesh = Cast<UStaticMesh>( SelectedSet->GetTop( UStaticMesh::StaticClass() ) );
-		if( SelectedStaticMesh != NULL )
+		FEdModeFoliage* Mode = (FEdModeFoliage*)GLevelEditorModeTools().GetActiveMode(FBuiltinEditorModes::EM_Foliage);
+
+		bool bMeshMerged = false;
+		if (Mode->ReplaceStaticMesh(FoliageSettingsPtr, SelectedStaticMesh, bMeshMerged))
 		{
-			FEdModeFoliage* Mode = (FEdModeFoliage*)GLevelEditorModeTools().GetActiveMode( FBuiltinEditorModes::EM_Foliage );
-			
-			bool bMeshMerged = false;
-			if(Mode->ReplaceStaticMesh(AssociatedStaticMesh, SelectedStaticMesh, bMeshMerged))
+			// If they were merged, simply remove the current item. Otherwise replace it.
+			if (bMeshMerged)
+			{
+				FoliageEditPtr.Pin()->RemoveItemFromScrollbox(SharedThis(this));
+			}
+			else
 			{
-				// If they were merged, simply remove the current item. Otherwise replace it.
-				if(bMeshMerged)
-				{
-					FoliageEditPtr.Pin()->RemoveItemFromScrollbox(SharedThis(this));
-				}
-				else
-				{
-					FoliageEditPtr.Pin()->ReplaceItem(SharedThis(this), SelectedStaticMesh);
-				}
+				FoliageEditPtr.Pin()->ReplaceItem(SharedThis(this), SelectedStaticMesh);
 			}
 		}
 	}
 
-
 	return FReply::Handled();
 }
 
 FSlateColor SFoliageEditMeshDisplayItem::GetBorderColor() const
 {
-	if(FoliageSettingsPtr->IsSelected)
+	if (FoliageSettingsPtr->IsSelected)
 	{
-		return 	FSlateColor(FLinearColor(0.828f, 0.364f, 0.003f));
+		return FSlateColor(FLinearColor(0.828f, 0.364f, 0.003f));
 	}
 
 	return FSlateColor(FLinearColor(0.25f, 0.25f, 0.25f));
@@ -1753,7 +1772,7 @@ FSlateColor SFoliageEditMeshDisplayItem::GetBorderColor() const
 
 ESlateCheckBoxState::Type SFoliageEditMeshDisplayItem::IsSelected() const
 {
-	return FoliageSettingsPtr->IsSelected? ESlateCheckBoxState::Checked : ESlateCheckBoxState::Unchecked;
+	return FoliageSettingsPtr->IsSelected ? ESlateCheckBoxState::Checked : ESlateCheckBoxState::Unchecked;
 }
 
 void SFoliageEditMeshDisplayItem::OnSelectionChanged(ESlateCheckBoxState::Type InType)
@@ -1764,7 +1783,9 @@ void SFoliageEditMeshDisplayItem::OnSelectionChanged(ESlateCheckBoxState::Type I
 FReply SFoliageEditMeshDisplayItem::OnSync()
 {
 	TArray<UObject*> Objects;
-	Objects.Add(AssociatedStaticMesh);
+
+	Objects.Add(FoliageSettingsPtr->GetStaticMesh());
+
 	GEditor->SyncBrowserToObjects(Objects);
 
 	return FReply::Handled();
@@ -1772,9 +1793,9 @@ FReply SFoliageEditMeshDisplayItem::OnSync()
 
 FReply SFoliageEditMeshDisplayItem::OnRemove()
 {
-	FEdModeFoliage* Mode = (FEdModeFoliage*)GLevelEditorModeTools().GetActiveMode( FBuiltinEditorModes::EM_Foliage );
-	
-	if(Mode->RemoveFoliageMesh(AssociatedStaticMesh))
+	FEdModeFoliage* Mode = (FEdModeFoliage*)GLevelEditorModeTools().GetActiveMode(FBuiltinEditorModes::EM_Foliage);
+
+	if (Mode->RemoveFoliageMesh(FoliageSettingsPtr))
 	{
 		FoliageEditPtr.Pin()->RemoveItemFromScrollbox(SharedThis(this));
 	}
@@ -1784,15 +1805,15 @@ FReply SFoliageEditMeshDisplayItem::OnRemove()
 
 FReply SFoliageEditMeshDisplayItem::OnSaveRemoveSettings()
 {
-	FEdModeFoliage* Mode = (FEdModeFoliage*)GLevelEditorModeTools().GetActiveMode( FBuiltinEditorModes::EM_Foliage );
+	FEdModeFoliage* Mode = (FEdModeFoliage*)GLevelEditorModeTools().GetActiveMode(FBuiltinEditorModes::EM_Foliage);
 
-	if(FoliageSettingsPtr->GetOuter()->IsA(UPackage::StaticClass()))
+	if (FoliageSettingsPtr->GetOuter()->IsA(UPackage::StaticClass()))
 	{
-		UInstancedFoliageSettings* NewSettings = NULL;
-		NewSettings = Mode->CopySettingsObject(AssociatedStaticMesh);
+		UFoliageType* NewSettings = NULL;
+		NewSettings = Mode->CopySettingsObject(FoliageSettingsPtr);
 
 		// Do not replace the current one if NULL is returned, just keep the old one.
-		if(NewSettings)
+		if (NewSettings)
 		{
 			FoliageSettingsPtr = NewSettings;
 		}
@@ -1800,21 +1821,21 @@ FReply SFoliageEditMeshDisplayItem::OnSaveRemoveSettings()
 	else
 	{
 		// Build default settings asset name and path
-		FString DefaultAsset = FPackageName::GetLongPackagePath(AssociatedStaticMesh->GetOutermost()->GetName()) + TEXT("/") + AssociatedStaticMesh->GetName() + TEXT("_settings");
+		FString DefaultAsset = FPackageName::GetLongPackagePath(FoliageSettingsPtr->GetStaticMesh()->GetOutermost()->GetName()) + TEXT("/") + FoliageSettingsPtr->GetStaticMesh()->GetName() + TEXT("_settings");
 
-		TSharedRef<SDlgPickAssetPath> SettingDlg = 
+		TSharedRef<SDlgPickAssetPath> SettingDlg =
 			SNew(SDlgPickAssetPath)
 			.Title(LOCTEXT("SettingsDialogTitle", "Choose Location for Foliage Settings Asset"))
 			.DefaultAssetPath(FText::FromString(DefaultAsset));
 
 		if (SettingDlg->ShowModal() != EAppReturnType::Cancel)
 		{
-			UInstancedFoliageSettings* NewSettings = NULL;
+			UFoliageType* NewSettings = NULL;
 
-			NewSettings = Mode->SaveSettingsObject(SettingDlg->GetFullAssetPath(), AssociatedStaticMesh);
+			NewSettings = Mode->SaveSettingsObject(SettingDlg->GetFullAssetPath(), FoliageSettingsPtr);
 
 			// Do not replace the current one if NULL is returned, just keep the old one.
-			if(NewSettings)
+			if (NewSettings)
 			{
 				FoliageSettingsPtr = NewSettings;
 			}
@@ -1827,87 +1848,87 @@ FReply SFoliageEditMeshDisplayItem::OnSaveRemoveSettings()
 FReply SFoliageEditMeshDisplayItem::OnOpenSettings()
 {
 	FEditorDelegates::LoadSelectedAssetsIfNeeded.Broadcast();
-	USelection* SelectedSet = GEditor->GetSelectedSet( UInstancedFoliageSettings::StaticClass() );
-	UInstancedFoliageSettings* SelectedSettings = Cast<UInstancedFoliageSettings>( SelectedSet->GetTop( UInstancedFoliageSettings::StaticClass() ) );
-	if( SelectedSettings != NULL )
+	USelection* SelectedSet = GEditor->GetSelectedSet(UFoliageType::StaticClass());
+	UFoliageType* SelectedSettings = SelectedSet->GetTop<UFoliageType>();
+	if (SelectedSettings != nullptr)
 	{
-		FEdModeFoliage* Mode = (FEdModeFoliage*)GLevelEditorModeTools().GetActiveMode( FBuiltinEditorModes::EM_Foliage );
-		Mode->ReplaceSettingsObject(AssociatedStaticMesh, SelectedSettings);
+		FEdModeFoliage* Mode = (FEdModeFoliage*)GLevelEditorModeTools().GetActiveMode(FBuiltinEditorModes::EM_Foliage);
+		Mode->ReplaceSettingsObject(FoliageSettingsPtr, SelectedSettings);
 		FoliageSettingsPtr = SelectedSettings;
 	}
 
 	return FReply::Handled();
 }
 
-FString SFoliageEditMeshDisplayItem::GetSaveRemoveSettingsTooltip() const
+FText SFoliageEditMeshDisplayItem::GetSaveRemoveSettingsTooltip() const
 {
 	// Remove Settings tooltip.
-	if(FoliageSettingsPtr->GetOuter()->IsA(UPackage::StaticClass()))
+	if (FoliageSettingsPtr->GetOuter()->IsA(UPackage::StaticClass()))
 	{
-		return NSLOCTEXT("FoliageEdMode", "RemoveSettings_Tooltip", "Do not store the foliage settings in a shared InstancedFoliageSettings object.").ToString();
+		return NSLOCTEXT("FoliageEdMode", "RemoveSettings_Tooltip", "Do not store the foliage settings in a shared InstancedFoliageSettings object.");
 	}
 
 	// Save settings tooltip.
-	return NSLOCTEXT("FoliageEdMode", "SaveSettings_Tooltip", "Save these settings as an InstancedFoliageSettings object stored in a package.").ToString();
+	return NSLOCTEXT("FoliageEdMode", "SaveSettings_Tooltip", "Save these settings as an InstancedFoliageSettings object stored in a package.");
 }
 
-bool SFoliageEditMeshDisplayItem::IsPropertyVisible( const FPropertyAndParent& PropertyAndParent ) const
+bool SFoliageEditMeshDisplayItem::IsPropertyVisible(const FPropertyAndParent& PropertyAndParent) const
 {
 	const FString Category = FObjectEditorUtils::GetCategory(&PropertyAndParent.Property);
 	return Category == TEXT("Clustering") || Category == TEXT("Culling") || Category == TEXT("Lighting") || Category == TEXT("Collision") || Category == TEXT("Rendering");
 }
 
-FString SFoliageEditMeshDisplayItem::GetInstanceCountString() const
+FText SFoliageEditMeshDisplayItem::GetInstanceCountString() const
 {
-	return FString::Printf(*LOCTEXT("InstanceCount_Value", "Instance Count: %d").ToString(), FoliageMeshUIInfo->MeshInfo->GetInstanceCount());
+	return FText::Format(LOCTEXT("InstanceCount_Value", "Instance Count: {0}"), FText::AsNumber(FoliageMeshUIInfo->MeshInfo->GetInstanceCount()));
 }
 
-FString SFoliageEditMeshDisplayItem::GetInstanceClusterCountString() const
+FText SFoliageEditMeshDisplayItem::GetInstanceClusterCountString() const
 {
-	return FString::Printf(*LOCTEXT("ClusterCount_Value", "Cluster Count: %d").ToString(), FoliageMeshUIInfo->MeshInfo->InstanceClusters.Num());
+	return FText::Format(LOCTEXT("ClusterCount_Value", "Cluster Count: {0}"), FText::AsNumber(FoliageMeshUIInfo->MeshInfo->InstanceClusters.Num()));
 }
 
 EVisibility SFoliageEditMeshDisplayItem::IsReapplySettingsVisible() const
 {
-	FEdModeFoliage* Mode = (FEdModeFoliage*)GLevelEditorModeTools().GetActiveMode( FBuiltinEditorModes::EM_Foliage );
-	return Mode->UISettings.GetReapplyToolSelected()? EVisibility::Visible : EVisibility::Collapsed;
+	FEdModeFoliage* Mode = (FEdModeFoliage*)GLevelEditorModeTools().GetActiveMode(FBuiltinEditorModes::EM_Foliage);
+	return Mode->UISettings.GetReapplyToolSelected() ? EVisibility::Visible : EVisibility::Collapsed;
 }
 
 EVisibility SFoliageEditMeshDisplayItem::IsReapplySettingsVisible_Dummy() const
 {
-	FEdModeFoliage* Mode = (FEdModeFoliage*)GLevelEditorModeTools().GetActiveMode( FBuiltinEditorModes::EM_Foliage );
-	return Mode->UISettings.GetReapplyToolSelected()? EVisibility::Hidden : EVisibility::Collapsed;
+	FEdModeFoliage* Mode = (FEdModeFoliage*)GLevelEditorModeTools().GetActiveMode(FBuiltinEditorModes::EM_Foliage);
+	return Mode->UISettings.GetReapplyToolSelected() ? EVisibility::Hidden : EVisibility::Collapsed;
 }
 
 EVisibility SFoliageEditMeshDisplayItem::IsNotReapplySettingsVisible() const
 {
-	FEdModeFoliage* Mode = (FEdModeFoliage*)GLevelEditorModeTools().GetActiveMode( FBuiltinEditorModes::EM_Foliage );
-	return Mode->UISettings.GetReapplyToolSelected()? EVisibility::Collapsed : EVisibility::Visible;
+	FEdModeFoliage* Mode = (FEdModeFoliage*)GLevelEditorModeTools().GetActiveMode(FBuiltinEditorModes::EM_Foliage);
+	return Mode->UISettings.GetReapplyToolSelected() ? EVisibility::Collapsed : EVisibility::Visible;
 }
 
 EVisibility SFoliageEditMeshDisplayItem::IsUniformScalingVisible() const
 {
-	return FoliageSettingsPtr->UniformScale? EVisibility::Visible : EVisibility::Collapsed;
+	return FoliageSettingsPtr->UniformScale ? EVisibility::Visible : EVisibility::Collapsed;
 }
 
 EVisibility SFoliageEditMeshDisplayItem::IsNonUniformScalingVisible() const
 {
-	return FoliageSettingsPtr->UniformScale? EVisibility::Collapsed : EVisibility::Visible;
+	return FoliageSettingsPtr->UniformScale ? EVisibility::Collapsed : EVisibility::Visible;
 }
 
 EVisibility SFoliageEditMeshDisplayItem::IsNonUniformScalingVisible_Dummy() const
 {
-	return FoliageSettingsPtr->UniformScale? EVisibility::Collapsed : EVisibility::Hidden;
+	return FoliageSettingsPtr->UniformScale ? EVisibility::Collapsed : EVisibility::Hidden;
 }
 
 void SFoliageEditMeshDisplayItem::OnDensityReapply(ESlateCheckBoxState::Type InState)
 {
-	FoliageSettingsPtr->ReapplyDensity = InState == ESlateCheckBoxState::Checked? true : false;
+	FoliageSettingsPtr->ReapplyDensity = InState == ESlateCheckBoxState::Checked ? true : false;
 }
 
 ESlateCheckBoxState::Type SFoliageEditMeshDisplayItem::IsDensityReapplyChecked() const
 {
-	return FoliageSettingsPtr->ReapplyDensity? ESlateCheckBoxState::Checked : ESlateCheckBoxState::Unchecked;
+	return FoliageSettingsPtr->ReapplyDensity ? ESlateCheckBoxState::Checked : ESlateCheckBoxState::Unchecked;
 }
 
 void SFoliageEditMeshDisplayItem::OnDensityChanged(float InValue)
@@ -1932,12 +1953,12 @@ float SFoliageEditMeshDisplayItem::GetDensityReapply() const
 
 void SFoliageEditMeshDisplayItem::OnRadiusReapply(ESlateCheckBoxState::Type InState)
 {
-	FoliageSettingsPtr->ReapplyRadius = InState == ESlateCheckBoxState::Checked? true : false;
+	FoliageSettingsPtr->ReapplyRadius = InState == ESlateCheckBoxState::Checked ? true : false;
 }
 
 ESlateCheckBoxState::Type SFoliageEditMeshDisplayItem::IsRadiusReapplyChecked() const
 {
-	return FoliageSettingsPtr->ReapplyRadius? ESlateCheckBoxState::Checked : ESlateCheckBoxState::Unchecked;
+	return FoliageSettingsPtr->ReapplyRadius ? ESlateCheckBoxState::Checked : ESlateCheckBoxState::Unchecked;
 }
 
 void SFoliageEditMeshDisplayItem::OnRadiusChanged(float InValue)
@@ -1952,17 +1973,17 @@ float SFoliageEditMeshDisplayItem::GetRadius() const
 
 void SFoliageEditMeshDisplayItem::OnAlignToNormalReapply(ESlateCheckBoxState::Type InState)
 {
-	FoliageSettingsPtr->ReapplyAlignToNormal = InState == ESlateCheckBoxState::Checked? true : false;
+	FoliageSettingsPtr->ReapplyAlignToNormal = InState == ESlateCheckBoxState::Checked ? true : false;
 }
 
 ESlateCheckBoxState::Type SFoliageEditMeshDisplayItem::IsAlignToNormalReapplyChecked() const
 {
-	return FoliageSettingsPtr->ReapplyAlignToNormal? ESlateCheckBoxState::Checked : ESlateCheckBoxState::Unchecked;
+	return FoliageSettingsPtr->ReapplyAlignToNormal ? ESlateCheckBoxState::Checked : ESlateCheckBoxState::Unchecked;
 }
 
 void SFoliageEditMeshDisplayItem::OnAlignToNormal(ESlateCheckBoxState::Type InState)
 {
-	if(InState == ESlateCheckBoxState::Checked)
+	if (InState == ESlateCheckBoxState::Checked)
 	{
 		FoliageSettingsPtr->AlignToNormal = true;
 	}
@@ -1974,12 +1995,12 @@ void SFoliageEditMeshDisplayItem::OnAlignToNormal(ESlateCheckBoxState::Type InSt
 
 ESlateCheckBoxState::Type SFoliageEditMeshDisplayItem::IsAlignToNormalChecked() const
 {
-	return FoliageSettingsPtr->AlignToNormal? ESlateCheckBoxState::Checked : ESlateCheckBoxState::Unchecked;
+	return FoliageSettingsPtr->AlignToNormal ? ESlateCheckBoxState::Checked : ESlateCheckBoxState::Unchecked;
 }
 
 EVisibility SFoliageEditMeshDisplayItem::IsAlignToNormalVisible() const
 {
-	return FoliageSettingsPtr->AlignToNormal? EVisibility::Visible : EVisibility::Collapsed;
+	return FoliageSettingsPtr->AlignToNormal ? EVisibility::Visible : EVisibility::Collapsed;
 }
 
 void SFoliageEditMeshDisplayItem::OnMaxAngleChanged(float InValue)
@@ -1994,17 +2015,17 @@ float SFoliageEditMeshDisplayItem::GetMaxAngle() const
 
 void SFoliageEditMeshDisplayItem::OnRandomYawReapply(ESlateCheckBoxState::Type InState)
 {
-	FoliageSettingsPtr->ReapplyRandomYaw = InState == ESlateCheckBoxState::Checked? true : false;
+	FoliageSettingsPtr->ReapplyRandomYaw = InState == ESlateCheckBoxState::Checked ? true : false;
 }
 
 ESlateCheckBoxState::Type SFoliageEditMeshDisplayItem::IsRandomYawReapplyChecked() const
 {
-	return FoliageSettingsPtr->ReapplyRandomYaw? ESlateCheckBoxState::Checked : ESlateCheckBoxState::Unchecked;
+	return FoliageSettingsPtr->ReapplyRandomYaw ? ESlateCheckBoxState::Checked : ESlateCheckBoxState::Unchecked;
 }
 
 void SFoliageEditMeshDisplayItem::OnRandomYaw(ESlateCheckBoxState::Type InState)
 {
-	if(InState == ESlateCheckBoxState::Checked)
+	if (InState == ESlateCheckBoxState::Checked)
 	{
 		FoliageSettingsPtr->RandomYaw = true;
 	}
@@ -2016,12 +2037,12 @@ void SFoliageEditMeshDisplayItem::OnRandomYaw(ESlateCheckBoxState::Type InState)
 
 ESlateCheckBoxState::Type SFoliageEditMeshDisplayItem::IsRandomYawChecked() const
 {
-	return FoliageSettingsPtr->RandomYaw? ESlateCheckBoxState::Checked : ESlateCheckBoxState::Unchecked;
+	return FoliageSettingsPtr->RandomYaw ? ESlateCheckBoxState::Checked : ESlateCheckBoxState::Unchecked;
 }
 
 void SFoliageEditMeshDisplayItem::OnUniformScale(ESlateCheckBoxState::Type InState)
 {
-	if(InState == ESlateCheckBoxState::Checked)
+	if (InState == ESlateCheckBoxState::Checked)
 	{
 		FoliageSettingsPtr->UniformScale = true;
 
@@ -2036,19 +2057,19 @@ void SFoliageEditMeshDisplayItem::OnUniformScale(ESlateCheckBoxState::Type InSta
 
 ESlateCheckBoxState::Type SFoliageEditMeshDisplayItem::IsUniformScaleChecked() const
 {
-	return FoliageSettingsPtr->UniformScale? ESlateCheckBoxState::Checked : ESlateCheckBoxState::Unchecked;
+	return FoliageSettingsPtr->UniformScale ? ESlateCheckBoxState::Checked : ESlateCheckBoxState::Unchecked;
 }
 
 void SFoliageEditMeshDisplayItem::OnScaleUniformReapply(ESlateCheckBoxState::Type InState)
 {
-	FoliageSettingsPtr->ReapplyScaleX = InState == ESlateCheckBoxState::Checked? true : false;
-	FoliageSettingsPtr->ReapplyScaleY = InState == ESlateCheckBoxState::Checked? true : false;
-	FoliageSettingsPtr->ReapplyScaleZ = InState == ESlateCheckBoxState::Checked? true : false;
+	FoliageSettingsPtr->ReapplyScaleX = InState == ESlateCheckBoxState::Checked ? true : false;
+	FoliageSettingsPtr->ReapplyScaleY = InState == ESlateCheckBoxState::Checked ? true : false;
+	FoliageSettingsPtr->ReapplyScaleZ = InState == ESlateCheckBoxState::Checked ? true : false;
 }
 
 ESlateCheckBoxState::Type SFoliageEditMeshDisplayItem::IsScaleUniformReapplyChecked() const
 {
-	return FoliageSettingsPtr->ReapplyScaleX? ESlateCheckBoxState::Checked : ESlateCheckBoxState::Unchecked;
+	return FoliageSettingsPtr->ReapplyScaleX ? ESlateCheckBoxState::Checked : ESlateCheckBoxState::Unchecked;
 }
 
 void SFoliageEditMeshDisplayItem::OnScaleUniformMinChanged(float InValue)
@@ -2090,7 +2111,7 @@ void SFoliageEditMeshDisplayItem::OnScaleXReapply(ESlateCheckBoxState::Type InSt
 
 ESlateCheckBoxState::Type SFoliageEditMeshDisplayItem::IsScaleXReapplyChecked() const
 {
-	return FoliageSettingsPtr->ReapplyScaleX? ESlateCheckBoxState::Checked : ESlateCheckBoxState::Unchecked;
+	return FoliageSettingsPtr->ReapplyScaleX ? ESlateCheckBoxState::Checked : ESlateCheckBoxState::Unchecked;
 }
 
 void SFoliageEditMeshDisplayItem::OnScaleXMinChanged(float InValue)
@@ -2173,12 +2194,12 @@ void SFoliageEditMeshDisplayItem::OnScaleYLocked(ESlateCheckBoxState::Type InSta
 
 void SFoliageEditMeshDisplayItem::OnScaleZReapply(ESlateCheckBoxState::Type InState)
 {
-	FoliageSettingsPtr->ReapplyScaleZ = InState == ESlateCheckBoxState::Checked? true : false;
+	FoliageSettingsPtr->ReapplyScaleZ = InState == ESlateCheckBoxState::Checked ? true : false;
 }
 
 ESlateCheckBoxState::Type SFoliageEditMeshDisplayItem::IsScaleZReapplyChecked() const
 {
-	return FoliageSettingsPtr->ReapplyScaleZ? ESlateCheckBoxState::Checked : ESlateCheckBoxState::Unchecked;
+	return FoliageSettingsPtr->ReapplyScaleZ ? ESlateCheckBoxState::Checked : ESlateCheckBoxState::Unchecked;
 }
 
 void SFoliageEditMeshDisplayItem::OnScaleZMinChanged(float InValue)
@@ -2217,12 +2238,12 @@ void SFoliageEditMeshDisplayItem::OnScaleZLocked(ESlateCheckBoxState::Type InSta
 
 void SFoliageEditMeshDisplayItem::OnZOffsetReapply(ESlateCheckBoxState::Type InState)
 {
-	FoliageSettingsPtr->ReapplyZOffset = InState == ESlateCheckBoxState::Checked? true : false;
+	FoliageSettingsPtr->ReapplyZOffset = InState == ESlateCheckBoxState::Checked ? true : false;
 }
 
 ESlateCheckBoxState::Type SFoliageEditMeshDisplayItem::IsZOffsetReapplyChecked() const
 {
-	return FoliageSettingsPtr->ReapplyZOffset? ESlateCheckBoxState::Checked : ESlateCheckBoxState::Unchecked;
+	return FoliageSettingsPtr->ReapplyZOffset ? ESlateCheckBoxState::Checked : ESlateCheckBoxState::Unchecked;
 }
 
 void SFoliageEditMeshDisplayItem::OnZOffsetMin(float InValue)
@@ -2251,12 +2272,12 @@ float SFoliageEditMeshDisplayItem::GetZOffsetMax() const
 
 void SFoliageEditMeshDisplayItem::OnRandomPitchReapply(ESlateCheckBoxState::Type InState)
 {
-	FoliageSettingsPtr->ReapplyRandomPitchAngle = InState == ESlateCheckBoxState::Checked? true : false;
+	FoliageSettingsPtr->ReapplyRandomPitchAngle = InState == ESlateCheckBoxState::Checked ? true : false;
 }
 
 ESlateCheckBoxState::Type SFoliageEditMeshDisplayItem::IsRandomPitchReapplyChecked() const
 {
-	return FoliageSettingsPtr->ReapplyRandomPitchAngle? ESlateCheckBoxState::Checked : ESlateCheckBoxState::Unchecked;
+	return FoliageSettingsPtr->ReapplyRandomPitchAngle ? ESlateCheckBoxState::Checked : ESlateCheckBoxState::Unchecked;
 }
 
 void SFoliageEditMeshDisplayItem::OnRandomPitchChanged(float InValue)
@@ -2271,12 +2292,12 @@ float SFoliageEditMeshDisplayItem::GetRandomPitch() const
 
 void SFoliageEditMeshDisplayItem::OnGroundSlopeReapply(ESlateCheckBoxState::Type InState)
 {
-	FoliageSettingsPtr->ReapplyGroundSlope = InState == ESlateCheckBoxState::Checked? true : false;
+	FoliageSettingsPtr->ReapplyGroundSlope = InState == ESlateCheckBoxState::Checked ? true : false;
 }
 
 ESlateCheckBoxState::Type SFoliageEditMeshDisplayItem::IsGroundSlopeReapplyChecked() const
 {
-	return FoliageSettingsPtr->ReapplyGroundSlope? ESlateCheckBoxState::Checked : ESlateCheckBoxState::Unchecked;
+	return FoliageSettingsPtr->ReapplyGroundSlope ? ESlateCheckBoxState::Checked : ESlateCheckBoxState::Unchecked;
 }
 
 void SFoliageEditMeshDisplayItem::OnGroundSlopeChanged(float InValue)
@@ -2291,12 +2312,12 @@ float SFoliageEditMeshDisplayItem::GetGroundSlope() const
 
 void SFoliageEditMeshDisplayItem::OnHeightReapply(ESlateCheckBoxState::Type InState)
 {
-	FoliageSettingsPtr->ReapplyHeight = InState == ESlateCheckBoxState::Checked? true : false;
+	FoliageSettingsPtr->ReapplyHeight = InState == ESlateCheckBoxState::Checked ? true : false;
 }
 
 ESlateCheckBoxState::Type SFoliageEditMeshDisplayItem::IsHeightReapplyChecked() const
 {
-	return FoliageSettingsPtr->ReapplyHeight? ESlateCheckBoxState::Checked : ESlateCheckBoxState::Unchecked;
+	return FoliageSettingsPtr->ReapplyHeight ? ESlateCheckBoxState::Checked : ESlateCheckBoxState::Unchecked;
 }
 
 void SFoliageEditMeshDisplayItem::OnHeightMinChanged(float InValue)
@@ -2325,17 +2346,17 @@ float SFoliageEditMeshDisplayItem::GetHeightMax() const
 
 void SFoliageEditMeshDisplayItem::OnLandscapeLayerReapply(ESlateCheckBoxState::Type InState)
 {
-	FoliageSettingsPtr->ReapplyLandscapeLayer = InState == ESlateCheckBoxState::Checked? true : false;
+	FoliageSettingsPtr->ReapplyLandscapeLayer = InState == ESlateCheckBoxState::Checked ? true : false;
 }
 
 ESlateCheckBoxState::Type SFoliageEditMeshDisplayItem::IsLandscapeLayerReapplyChecked() const
 {
-	return FoliageSettingsPtr->ReapplyLandscapeLayer? ESlateCheckBoxState::Checked : ESlateCheckBoxState::Unchecked;
+	return FoliageSettingsPtr->ReapplyLandscapeLayer ? ESlateCheckBoxState::Checked : ESlateCheckBoxState::Unchecked;
 }
 
 void SFoliageEditMeshDisplayItem::OnLandscapeLayerChanged(const FText& InValue)
 {
-	FoliageSettingsPtr->LandscapeLayer = FName(  *InValue.ToString() );
+	FoliageSettingsPtr->LandscapeLayer = FName(*InValue.ToString());
 }
 
 FText SFoliageEditMeshDisplayItem::GetLandscapeLayer() const
@@ -2345,7 +2366,7 @@ FText SFoliageEditMeshDisplayItem::GetLandscapeLayer() const
 
 void SFoliageEditMeshDisplayItem::OnCollisionWithWorld(ESlateCheckBoxState::Type InState)
 {
-	if(InState == ESlateCheckBoxState::Checked)
+	if (InState == ESlateCheckBoxState::Checked)
 	{
 		FoliageSettingsPtr->CollisionWithWorld = true;
 	}
@@ -2357,22 +2378,22 @@ void SFoliageEditMeshDisplayItem::OnCollisionWithWorld(ESlateCheckBoxState::Type
 
 ESlateCheckBoxState::Type SFoliageEditMeshDisplayItem::IsCollisionWithWorldChecked() const
 {
-	return FoliageSettingsPtr->CollisionWithWorld? ESlateCheckBoxState::Checked : ESlateCheckBoxState::Unchecked;
+	return FoliageSettingsPtr->CollisionWithWorld ? ESlateCheckBoxState::Checked : ESlateCheckBoxState::Unchecked;
 }
 
 EVisibility SFoliageEditMeshDisplayItem::IsCollisionWithWorldVisible() const
 {
-	return FoliageSettingsPtr->CollisionWithWorld? EVisibility::Visible : EVisibility::Collapsed;
+	return FoliageSettingsPtr->CollisionWithWorld ? EVisibility::Visible : EVisibility::Collapsed;
 }
 
 void SFoliageEditMeshDisplayItem::OnCollisionWithWorldReapply(ESlateCheckBoxState::Type InState)
 {
-	FoliageSettingsPtr->ReapplyCollisionWithWorld = InState == ESlateCheckBoxState::Checked? true : false;
+	FoliageSettingsPtr->ReapplyCollisionWithWorld = InState == ESlateCheckBoxState::Checked ? true : false;
 }
 
 ESlateCheckBoxState::Type SFoliageEditMeshDisplayItem::IsCollisionWithWorldReapplyChecked() const
 {
-	return FoliageSettingsPtr->ReapplyCollisionWithWorld? ESlateCheckBoxState::Checked : ESlateCheckBoxState::Unchecked;
+	return FoliageSettingsPtr->ReapplyCollisionWithWorld ? ESlateCheckBoxState::Checked : ESlateCheckBoxState::Unchecked;
 }
 
 void SFoliageEditMeshDisplayItem::OnCollisionScaleXChanged(float InValue)
@@ -2412,7 +2433,7 @@ void SFoliageEditMeshDisplayItem::OnVertexColorMask(ESlateCheckBoxState::Type In
 
 ESlateCheckBoxState::Type SFoliageEditMeshDisplayItem::IsVertexColorMaskChecked(FoliageVertexColorMask Mask) const
 {
-	return FoliageSettingsPtr->VertexColorMask == Mask ?  ESlateCheckBoxState::Checked : ESlateCheckBoxState::Unchecked;
+	return FoliageSettingsPtr->VertexColorMask == Mask ? ESlateCheckBoxState::Checked : ESlateCheckBoxState::Unchecked;
 }
 
 void SFoliageEditMeshDisplayItem::OnVertexColorMaskThresholdChanged(float InValue)
@@ -2442,15 +2463,15 @@ ESlateCheckBoxState::Type SFoliageEditMeshDisplayItem::IsVertexColorMaskInvertCh
 
 void SFoliageEditMeshDisplayItem::OnVertexColorMaskReapply(ESlateCheckBoxState::Type InState)
 {
-	FoliageSettingsPtr->ReapplyVertexColorMask = InState == ESlateCheckBoxState::Checked? true : false;
+	FoliageSettingsPtr->ReapplyVertexColorMask = InState == ESlateCheckBoxState::Checked ? true : false;
 }
 
 ESlateCheckBoxState::Type SFoliageEditMeshDisplayItem::IsVertexColorMaskReapplyChecked() const
 {
-	return FoliageSettingsPtr->ReapplyVertexColorMask? ESlateCheckBoxState::Checked : ESlateCheckBoxState::Unchecked;
+	return FoliageSettingsPtr->ReapplyVertexColorMask ? ESlateCheckBoxState::Checked : ESlateCheckBoxState::Unchecked;
 }
 
-FReply SFoliageEditMeshDisplayItem::OnMouseDownSelection( const FGeometry& MyGeometry, const FPointerEvent& MouseEvent )
+FReply SFoliageEditMeshDisplayItem::OnMouseDownSelection(const FGeometry& MyGeometry, const FPointerEvent& MouseEvent)
 {
 	FoliageSettingsPtr->IsSelected = !FoliageSettingsPtr->IsSelected;
 	return FReply::Handled();
@@ -2458,7 +2479,7 @@ FReply SFoliageEditMeshDisplayItem::OnMouseDownSelection( const FGeometry& MyGeo
 
 const FSlateBrush* SFoliageEditMeshDisplayItem::GetSaveSettingsBrush() const
 {
-	if(FoliageSettingsPtr->GetOuter()->IsA(UPackage::StaticClass()))
+	if (FoliageSettingsPtr->GetOuter()->IsA(UPackage::StaticClass()))
 	{
 		return FEditorStyle::GetBrush(TEXT("FoliageEditMode.DeleteItem"));
 	}
diff --git a/Engine/Source/Editor/FoliageEdit/Private/SFoliageEditMeshDisplayItem.h b/Engine/Source/Editor/FoliageEdit/Private/SFoliageEditMeshDisplayItem.h
--- a/Engine/Source/Editor/FoliageEdit/Private/SFoliageEditMeshDisplayItem.h
+++ b/Engine/Source/Editor/FoliageEdit/Private/SFoliageEditMeshDisplayItem.h
@@ -2,6 +2,10 @@
 
 #pragma once
 
+// Forward declarations
+class IDetailsView;
+struct FPropertyAndParent;
+
 namespace ECurrentViewSettings
 {
 	enum Type 
@@ -15,20 +19,14 @@ namespace ECurrentViewSettings
 class SFoliageEditMeshDisplayItem : public SBorder, public FNotifyHook
 {
 public:
-	SLATE_BEGIN_ARGS( SFoliageEditMeshDisplayItem )
+	SLATE_BEGIN_ARGS(SFoliageEditMeshDisplayItem)
 		: _FoliageSettingsPtr(NULL)
-		, _AssociatedStaticMesh(NULL)
 		{}
 
-		SLATE_ARGUMENT( UInstancedFoliageSettings*, FoliageSettingsPtr )
-
-		SLATE_ARGUMENT( UStaticMesh*, AssociatedStaticMesh )
-
-		SLATE_ARGUMENT( TSharedPtr<FAssetThumbnail>, AssetThumbnail )
-
-		SLATE_ARGUMENT( TSharedPtr< SFoliageEdit >, FoliageEditPtr )
-
-		SLATE_ARGUMENT( TSharedPtr<FFoliageMeshUIInfo>, FoliageMeshUIInfo )
+		SLATE_ARGUMENT(UFoliageType*, FoliageSettingsPtr)
+		SLATE_ARGUMENT(TSharedPtr<FAssetThumbnail>, AssetThumbnail)
+		SLATE_ARGUMENT(TWeakPtr<SFoliageEdit>, FoliageEditPtr)
+		SLATE_ARGUMENT(TSharedPtr<FFoliageMeshUIInfo>, FoliageMeshUIInfo)
 	SLATE_END_ARGS()
 
 	/**
@@ -36,22 +34,21 @@ public:
 	 *
 	 * @param InArgs   A declaration from which to construct the widget
 	 */
-	void Construct( const FArguments& InArgs );
+	void Construct(const FArguments& InArgs);
 
 	/**
 	 * Called when the replace button is pressed, simply updates it's current data with the passed in data.
 	 *
 	 * @param InFoliageSettingsPtr			Pointer to the settings for this item.
-	 * @param InAssociatedStaticMesh		The static mesh this item is associated with.
 	 * @param InAssetThumbnail				The thumbnail to display.
 	 * @param InFoliageMeshUIInfo			The UI info for this item.
 	 */
-	void Replace(UInstancedFoliageSettings* InFoliageSettingsPtr, UStaticMesh* InAssociatedStaticMesh, TSharedPtr<FAssetThumbnail> InAssetThumbnail, TSharedPtr<FFoliageMeshUIInfo> InFoliageMeshUIInfo);
+	void Replace(UFoliageType* InFoliageSettingsPtr, TSharedPtr<FAssetThumbnail> InAssetThumbnail, TSharedPtr<FFoliageMeshUIInfo> InFoliageMeshUIInfo);
 
 	/** Used to cache and return the display status of items when undoing. */
 	ECurrentViewSettings::Type GetCurrentDisplayStatus() const;
 
-	/** 
+	/**
 	 * Sets the display status of the item.
 	 *
 	 * @param InDisplayStatus		The status to set the item's display to.
@@ -59,14 +56,14 @@ public:
 	void SetCurrentDisplayStatus(ECurrentViewSettings::Type InDisplayStatus);
 
 	/** Handles clicking on the selection area for the item. */
-	FReply OnMouseDownSelection( const FGeometry& MyGeometry, const FPointerEvent& MouseEvent );
+	FReply OnMouseDownSelection(const FGeometry& MyGeometry, const FPointerEvent& MouseEvent);
 
 	/** Returns the current brush for the Save/Remove settings button. */
 	const FSlateBrush* GetSaveSettingsBrush() const;
 
 private:
 	/** Notification when the cluster details are changed */
-	void NotifyPostChange( const FPropertyChangedEvent& PropertyChangedEvent, UProperty* PropertyThatChanged ) override;
+	void NotifyPostChange(const FPropertyChangedEvent& PropertyChangedEvent, UProperty* PropertyThatChanged) override;
 
 	/** Callback to get the border color for the item. */
 	FSlateColor GetBorderColor() const;
@@ -93,9 +90,9 @@ private:
 	bool IsClusterSettingsEnabled() const;
 
 	/** General UI delegate functions. */
-	FString GetStaticMeshname() const;
+	FText GetStaticMeshname() const;
 
-	FString GetSettingsLabelText() const;
+	FText GetSettingsLabelText() const;
 
 	EVisibility  IsNoSettingsVisible() const;
 
@@ -117,15 +114,15 @@ private:
 	FReply OnOpenSettings();
 
 	/** Returns the tool-tip to use for the save settings button. */
-	FString GetSaveRemoveSettingsTooltip() const;
+	FText GetSaveRemoveSettingsTooltip() const;
 
 	/** Determines what properties to display in the details panel. */
-	bool IsPropertyVisible( const struct FPropertyAndParent& PropertyAndParent ) const;
+	bool IsPropertyVisible(const FPropertyAndParent& PropertyAndParent) const;
 
 	/** Cluster Accessors */
-	FString GetInstanceCountString() const;
+	FText GetInstanceCountString() const;
 
-	FString GetInstanceClusterCountString() const;
+	FText GetInstanceClusterCountString() const;
 
 	/** General functions for delegates. */
 	EVisibility IsReapplySettingsVisible() const;
@@ -141,7 +138,7 @@ private:
 	void OnDensityReapply(ESlateCheckBoxState::Type InState);
 
 	ESlateCheckBoxState::Type IsDensityReapplyChecked() const;
-	
+
 	void OnDensityChanged(float InValue);
 
 	float GetDensity() const;
@@ -338,7 +335,7 @@ private:
 	void OnVertexColorMaskThresholdChanged(float InValue);
 
 	float GetVertexColorMaskThreshold() const;
-	
+
 	EVisibility IsVertexColorMaskThresholdVisible() const;
 
 	void OnVertexColorMaskInvert(ESlateCheckBoxState::Type InState);
@@ -347,13 +344,10 @@ private:
 
 private:
 	/** Pointer back to the foliage edit mode containing this. */
-	TWeakPtr< SFoliageEdit > FoliageEditPtr;
+	TWeakPtr<SFoliageEdit> FoliageEditPtr;
 
 	/** The settings for placing foliage. */
-	UInstancedFoliageSettings* FoliageSettingsPtr;
-
-	/** The static mesh this item is associated with. */
-	UStaticMesh* AssociatedStaticMesh;
+	UFoliageType* FoliageSettingsPtr;
 
 	/** Contains the thumbnail widget, needed to recreate the thumbnail when replacing. */
 	TSharedPtr<SBorder> ThumbnailWidgetBorder;
@@ -373,9 +367,9 @@ private:
 	ECurrentViewSettings::Type CurrentViewSettings;
 
 	/** Details panel for cluster settings. */
-	TSharedPtr<class IDetailsView> ClusterSettingsDetails;
+	TSharedPtr<IDetailsView> ClusterSettingsDetails;
 
-	/** Holds the UInstancedFoliageSettings object for the details panel. */
+	/** Holds the UFoliageType object for the details panel. */
 	TArray<UObject*> DetailsObjectList;
 
 	/** Command list for binding functions to for the toolbar. */
diff --git a/Engine/Source/Editor/UnrealEd/Private/EditorExporters.cpp b/Engine/Source/Editor/UnrealEd/Private/EditorExporters.cpp
--- a/Engine/Source/Editor/UnrealEd/Private/EditorExporters.cpp
+++ b/Engine/Source/Editor/UnrealEd/Private/EditorExporters.cpp
@@ -19,6 +19,7 @@
 #include "ImageUtils.h"
 #include "Foliage/InstancedFoliageActor.h"
 #include "Landscape/Landscape.h"
+#include "Foliage/FoliageType.h"
 #include "UnrealExporter.h"
 #include "EngineModule.h"
 
@@ -503,8 +504,8 @@ void ULevelExporterT3D::ExportComponentExtra(const FExportObjectInnerContext* Co
 				AInstancedFoliageActor* IFA = AInstancedFoliageActor::GetInstancedFoliageActorForLevel(ComponentLevel);
 				if (IFA)
 				{
-					TMap<class UStaticMesh*, TArray<const FFoliageInstancePlacementInfo*>> FoliageInstanceMap = IFA->GetInstancesForComponent(ActorComponent);
-					for (const TPair<class UStaticMesh*, TArray<const FFoliageInstancePlacementInfo*>>& MapEntry : FoliageInstanceMap)
+					auto FoliageInstanceMap = IFA->GetInstancesForComponent(ActorComponent);
+					for (const auto& MapEntry : FoliageInstanceMap)
 					{
 						Ar.Logf(TEXT("%sBegin Foliage StaticMesh=%s Component=%s%s"), FCString::Spc(TextIndent), *MapEntry.Key->GetPathName(), *ActorComponent->GetName(), LINE_TERMINATOR);
 						for (const FFoliageInstancePlacementInfo* Inst : MapEntry.Value)
diff --git a/Engine/Source/Editor/UnrealEd/Private/EditorObject.cpp b/Engine/Source/Editor/UnrealEd/Private/EditorObject.cpp
--- a/Engine/Source/Editor/UnrealEd/Private/EditorObject.cpp
+++ b/Engine/Source/Editor/UnrealEd/Private/EditorObject.cpp
@@ -9,6 +9,7 @@
 #include "BSPOps.h"
 
 #include "Foliage/InstancedFoliageActor.h"
+#include "Foliage/FoliageType.h"
 
 DEFINE_LOG_CATEGORY_STATIC(LogEditorObject, Log, All);
 
@@ -231,12 +232,12 @@ static const TCHAR* ImportProperties(
 				ImportedBrush = 1;
 			}
 		}
-		else if( GetBEGIN(&Str,TEXT("Foliage")) )
+		else if (GetBEGIN(&Str, TEXT("Foliage")))
 		{
 			UStaticMesh* StaticMesh;
 			FName ComponentName;
 			if (SubobjectRoot &&
-				ParseObject<UStaticMesh>( Str, TEXT("StaticMesh="), StaticMesh, ANY_PACKAGE ) &&
+				ParseObject<UStaticMesh>(Str, TEXT("StaticMesh="), StaticMesh, ANY_PACKAGE) &&
 				FParse::Value(Str, TEXT("Component="), ComponentName) )
 			{
 				UActorComponent* ActorComponent = FindObjectFast<UActorComponent>(SubobjectRoot, ComponentName);
@@ -246,14 +247,14 @@ static const TCHAR* ImportProperties(
 					ULevel* ComponentLevel = CastChecked<ULevel>(SubobjectRoot->GetOuter());
 					if (ComponentLevel->IsCurrentLevel())
 					{
-						AInstancedFoliageActor* IFA = AInstancedFoliageActor::GetInstancedFoliageActor(ComponentLevel->OwningWorld);
+						AInstancedFoliageActor* IFA = AInstancedFoliageActor::GetInstancedFoliageActorForCurrentLevel(ComponentLevel->OwningWorld);
 
 						const TCHAR* StrPtr;
 						FString TextLine;
-						while( FParse::Line( &SourceText, TextLine ) )
+						while (FParse::Line(&SourceText, TextLine))
 						{
 							StrPtr = *TextLine;
-							if( GetEND(&StrPtr,TEXT("Foliage")) )
+							if (GetEND(&StrPtr, TEXT("Foliage")))
 							{
 								break;
 							}
@@ -261,29 +262,34 @@ static const TCHAR* ImportProperties(
 							// Parse the instance properties
 							FFoliageInstance Instance;
 							FString Temp;
-							if( FParse::Value(StrPtr,TEXT("Location="), Temp, false) )
+							if (FParse::Value(StrPtr, TEXT("Location="), Temp, false))
 							{
 								GetFVECTOR(*Temp, Instance.Location);
 							}
-							if( FParse::Value(StrPtr,TEXT("Rotation="), Temp, false) )
+							if (FParse::Value(StrPtr, TEXT("Rotation="), Temp, false))
 							{
-								GetFROTATOR(*Temp, Instance.Rotation,1);
+								GetFROTATOR(*Temp, Instance.Rotation, 1);
 							}
-							if( FParse::Value(StrPtr,TEXT("PreAlignRotation="), Temp, false) )
+							if (FParse::Value(StrPtr, TEXT("PreAlignRotation="), Temp, false))
 							{
-								GetFROTATOR(*Temp, Instance.PreAlignRotation,1);
+								GetFROTATOR(*Temp, Instance.PreAlignRotation, 1);
 							}
-							if( FParse::Value(StrPtr,TEXT("DrawScale3D="), Temp, false) )
+							if (FParse::Value(StrPtr, TEXT("DrawScale3D="), Temp, false))
 							{
 								GetFVECTOR(*Temp, Instance.DrawScale3D);
 							}
-							FParse::Value(StrPtr,TEXT("Flags="),Instance.Flags);
+							FParse::Value(StrPtr, TEXT("Flags="), Instance.Flags);
 
 							Instance.Base = ActorComponent;
 
 							// Add the instance
-							FFoliageMeshInfo* MeshInfo = IFA->FindOrAddMesh(StaticMesh);
-							MeshInfo->AddInstance(IFA, StaticMesh, Instance);
+							FFoliageMeshInfo* MeshInfo;
+							UFoliageType* Type = IFA->GetSettingsForMesh(StaticMesh, &MeshInfo);
+							if (Type == NULL)
+							{
+								MeshInfo = IFA->AddMesh(StaticMesh, &Type);
+							}
+							MeshInfo->AddInstance(IFA, Type, Instance);
 						}
 					}
 				}
diff --git a/Engine/Source/Runtime/Core/Public/Templates/UniqueObj.h b/Engine/Source/Runtime/Core/Public/Templates/UniqueObj.h
--- a/Engine/Source/Runtime/Core/Public/Templates/UniqueObj.h
+++ b/Engine/Source/Runtime/Core/Public/Templates/UniqueObj.h
@@ -21,6 +21,13 @@ public:
 	{
 	}
 
+	// As TUniqueObj's internal pointer can never be null, we can't move that
+	// On the other hand we can call the move constructor of "T"
+	TUniqueObj(TUniqueObj&& other)
+		: Obj(MakeUnique<T>(MoveTemp(*other)))
+	{
+	}
+
 	template <typename Arg>
 	explicit TUniqueObj(Arg&& arg)
 		: Obj(MakeUnique<T>(Forward<Arg>(arg)))
@@ -35,6 +42,7 @@ public:
 	}
 #endif
 
+	// Disallow copy-assignment
 #if PLATFORM_COMPILER_HAS_DEFAULTED_FUNCTIONS
 	TUniqueObj& operator=(const TUniqueObj&) = delete;
 #else
@@ -43,6 +51,13 @@ private:
 public:
 #endif
 
+	// Move-assignment is implemented as swapping the internal pointer
+	TUniqueObj& operator=(TUniqueObj&& other)
+	{
+		Swap(Obj, other.Obj);
+		return *this;
+	}
+
 	template <typename Arg>
 	TUniqueObj& operator=(Arg&& other)
 	{
diff --git a/Engine/Source/Runtime/Core/Public/Templates/UnrealTemplate.h b/Engine/Source/Runtime/Core/Public/Templates/UnrealTemplate.h
--- a/Engine/Source/Runtime/Core/Public/Templates/UnrealTemplate.h
+++ b/Engine/Source/Runtime/Core/Public/Templates/UnrealTemplate.h
@@ -87,9 +87,9 @@ template< class T > inline T AlignArbitrary( const T Ptr, uint32 Alignment )
 }
 template< class T > inline void Swap( T& A, T& B )
 {
-	const T Temp = A;
-	A = B;
-	B = Temp;
+	T Temp = MoveTemp(A);
+	A = MoveTemp(B);
+	B = MoveTemp(Temp);
 }
 template< class T > inline void Exchange( T& A, T& B )
 {
diff --git a/Engine/Source/Runtime/Core/Public/UObject/ObjectVersion.h b/Engine/Source/Runtime/Core/Public/UObject/ObjectVersion.h
--- a/Engine/Source/Runtime/Core/Public/UObject/ObjectVersion.h
+++ b/Engine/Source/Runtime/Core/Public/UObject/ObjectVersion.h
@@ -566,6 +566,8 @@ enum EUnrealEngineObjectUE4Version
 	VER_UE4_RENAME_CANBECHARACTERBASE,
 	// Change GameplayTagContainers to have FGameplayTags instead of FNames; Required to fix-up native serialization
 	VER_UE4_GAMEPLAY_TAG_CONTAINER_TAG_TYPE_CHANGE,
+	// Change from UInstancedFoliageSettings to UFoliageType, and change the api from being keyed on UStaticMesh* to UFoliageType*
+	VER_UE4_FOLIAGE_SETTINGS_TYPE,
 
 	// -----<new versions can be added before this line>-------------------------------------------------
 	// - this needs to be the last line (see note below)
