Fix crash when force-deleting a mesh used by foliage. Note: Still not a good idea!

diff --git a/Engine/Source/Editor/FoliageEdit/Private/SFoliageEditMeshDisplayItem.cpp b/Engine/Source/Editor/FoliageEdit/Private/SFoliageEditMeshDisplayItem.cpp
--- a/Engine/Source/Editor/FoliageEdit/Private/SFoliageEditMeshDisplayItem.cpp
+++ b/Engine/Source/Editor/FoliageEdit/Private/SFoliageEditMeshDisplayItem.cpp
@@ -1704,7 +1704,8 @@ bool SFoliageEditMeshDisplayItem::IsClusterSettingsEnabled() const
 
 FText SFoliageEditMeshDisplayItem::GetStaticMeshname() const
 {
-	return FText::FromName(FoliageSettingsPtr->GetStaticMesh()->GetFName());
+	UStaticMesh* Mesh = FoliageSettingsPtr->GetStaticMesh();
+	return FText::FromName(Mesh ? Mesh->GetFName() : NAME_None);
 }
 
 FText SFoliageEditMeshDisplayItem::GetSettingsLabelText() const
