Fix foliage placing on translucent and "no collision" meshes.

diff --git a/Engine/Config/BaseEditorUserSettings.ini b/Engine/Config/BaseEditorUserSettings.ini
--- a/Engine/Config/BaseEditorUserSettings.ini
+++ b/Engine/Config/BaseEditorUserSettings.ini
@@ -625,6 +625,7 @@ UnpaintDensity=0.000000
 bFilterLandscape=True
 bFilterStaticMesh=True
 bFilterBSP=True
+bFilterTranslucent=False
 
 
 [MeshPaintEdit]
diff --git a/Engine/Source/Editor/FoliageEdit/Private/FoliageEdMode.cpp b/Engine/Source/Editor/FoliageEdit/Private/FoliageEdMode.cpp
--- a/Engine/Source/Editor/FoliageEdit/Private/FoliageEdMode.cpp
+++ b/Engine/Source/Editor/FoliageEdit/Private/FoliageEdMode.cpp
@@ -51,6 +51,8 @@ FEdModeFoliage::FEdModeFoliage()
 	}
 
 	SphereBrushComponent = ConstructObject<UStaticMeshComponent>(UStaticMeshComponent::StaticClass());
+	SphereBrushComponent->SetCollisionProfileName(UCollisionProfile::NoCollision_ProfileName);
+	SphereBrushComponent->SetCollisionObjectType(ECC_WorldDynamic);
 	SphereBrushComponent->StaticMesh = StaticMesh;
 	SphereBrushComponent->Materials.Add(BrushMaterial);
 	SphereBrushComponent->SetAbsolute(true, true, true);
@@ -274,6 +276,35 @@ void FEdModeFoliage::Tick(FEditorViewportClient* ViewportClient, float DeltaTime
 	}
 }
 
+static bool FoliageTrace(UWorld* InWorld, FHitResult& OutHit, FVector InStart, FVector InEnd, FName InTraceTag, bool InbReturnFaceIndex = false)
+{
+	FCollisionQueryParams QueryParams(InTraceTag, true, AInstancedFoliageActor::GetInstancedFoliageActorForCurrentLevel(InWorld, false));
+	QueryParams.bReturnFaceIndex = InbReturnFaceIndex;
+
+	bool bResult = true;
+	while (true)
+	{
+		bResult = InWorld->LineTraceSingle(OutHit, InStart, InEnd, QueryParams, FCollisionObjectQueryParams(ECC_WorldStatic));
+		if (bResult)
+		{
+			// In the editor traces can hit "No Collision" type actors, so ugh.
+			FBodyInstance* BodyInstance = OutHit.Component->GetBodyInstance();
+			if (BodyInstance->GetCollisionEnabled() != ECollisionEnabled::QueryAndPhysics || BodyInstance->GetResponseToChannel(ECC_WorldStatic) != ECR_Block)
+			{
+				AActor* Actor = OutHit.Actor.Get();
+				if (Actor)
+				{
+					QueryParams.AddIgnoredActor(Actor);
+				}
+				InStart = OutHit.ImpactPoint;
+				continue;
+			}
+		}
+		break;
+	}
+	return bResult;
+}
+
 /** Trace under the mouse cursor and update brush position */
 void FEdModeFoliage::FoliageBrushTrace(FEditorViewportClient* ViewportClient, int32 MouseX, int32 MouseY)
 {
@@ -295,20 +326,20 @@ void FEdModeFoliage::FoliageBrushTrace(FEditorViewportClient* ViewportClient, in
 
 		FHitResult Hit;
 		UWorld* World = ViewportClient->GetWorld();
-		if (World->LineTraceSingle(Hit, Start, End, FCollisionQueryParams(NAME_None, true, AInstancedFoliageActor::GetInstancedFoliageActorForCurrentLevel(World, false)), FCollisionObjectQueryParams(ECC_WorldStatic)))
+		static FName NAME_FoliageBrush = FName(TEXT("FoliageBrush"));
+		if (FoliageTrace(World, Hit, Start, End, NAME_FoliageBrush))
 		{
 			// Check filters
 			UPrimitiveComponent* PrimComp = Hit.Component.Get();
+			UMaterialInterface* Material = PrimComp ? PrimComp->GetMaterial(0) : nullptr;
+
 			if (PrimComp &&
-				(PrimComp->GetOutermost() != World->GetCurrentLevel()->GetOutermost() ||
-				(!UISettings.bFilterLandscape && PrimComp->IsA(ULandscapeHeightfieldCollisionComponent::StaticClass())) ||
-				(!UISettings.bFilterStaticMesh && PrimComp->IsA(UStaticMeshComponent::StaticClass())) ||
-				(!UISettings.bFilterBSP && PrimComp->IsA(UModelComponent::StaticClass()))
-				))
-			{
-				bBrushTraceValid = false;
-			}
-			else
+				PrimComp->GetOutermost() == World->GetCurrentLevel()->GetOutermost() &&
+				(UISettings.bFilterLandscape || !PrimComp->IsA(ULandscapeHeightfieldCollisionComponent::StaticClass())) &&
+				(UISettings.bFilterStaticMesh || !PrimComp->IsA(UStaticMeshComponent::StaticClass())) &&
+				(UISettings.bFilterBSP || !PrimComp->IsA(UModelComponent::StaticClass())) &&
+				(UISettings.bFilterTranslucent || !Material || !IsTranslucentBlendMode(Material->GetBlendMode()))
+				)
 			{
 				// Adjust the sphere brush
 				BrushLocation = Hit.Location;
@@ -393,7 +424,7 @@ bool CheckCollisionWithWorld(UFoliageType* Settings, const FFoliageInstance& Ins
 				FVector SamplePos = InstTransform.TransformPosition(FVector(Settings->LowBoundOriginRadius.X, Settings->LowBoundOriginRadius.Y, 2.f) + LocalSamplePos[i]);
 				float WorldRadius = (Settings->LowBoundOriginRadius.Z + 2.f)*FMath::Max(Inst.DrawScale3D.X, Inst.DrawScale3D.Y);
 				FVector NormalVector = Settings->AlignToNormal ? HitNormal : FVector(0, 0, 1);
-				if (InWorld->LineTraceSingle(Hit, SamplePos, SamplePos - NormalVector*WorldRadius, FCollisionQueryParams(true), FCollisionObjectQueryParams(ECC_WorldStatic)))
+				if (FoliageTrace(InWorld, Hit, SamplePos, SamplePos - NormalVector*WorldRadius, NAME_None))
 				{
 					if (LocalHit.Z - Inst.ZOffset < Settings->LowBoundOriginRadius.Z)
 					{
@@ -640,22 +671,19 @@ void FEdModeFoliage::AddInstancesForBrush(UWorld* InWorld, AInstancedFoliageActo
 
 			FHitResult Hit;
 			static FName NAME_AddInstancesForBrush = FName(TEXT("AddInstancesForBrush"));
-			FCollisionQueryParams TraceParams(NAME_AddInstancesForBrush, true, IFA);
-			if (Settings->VertexColorMask != FOLIAGEVERTEXCOLORMASK_Disabled)
-			{
-				TraceParams.bReturnFaceIndex = true;
-			}
-
-			if (InWorld->LineTraceSingle(Hit, Start, End, TraceParams, FCollisionObjectQueryParams(ECC_WorldStatic)))
+			if (FoliageTrace(InWorld, Hit, Start, End, NAME_AddInstancesForBrush, true))
 			{
 				// Check filters
 				UPrimitiveComponent* PrimComp = Hit.Component.Get();
-				if ((PrimComp &&
-					(PrimComp->GetOutermost() != InWorld->GetCurrentLevel()->GetOutermost() ||
+				UMaterialInterface* Material = PrimComp ? PrimComp->GetMaterial(0) : nullptr;
+
+				if (!PrimComp ||
+					PrimComp->GetOutermost() != InWorld->GetCurrentLevel()->GetOutermost() ||
 					(!UISettings.bFilterLandscape && PrimComp->IsA(ULandscapeHeightfieldCollisionComponent::StaticClass())) ||
 					(!UISettings.bFilterStaticMesh && PrimComp->IsA(UStaticMeshComponent::StaticClass())) ||
-					(!UISettings.bFilterBSP && PrimComp->IsA(UModelComponent::StaticClass()))
-					)))
+					(!UISettings.bFilterBSP && PrimComp->IsA(UModelComponent::StaticClass())) ||
+					(!UISettings.bFilterTranslucent && Material && IsTranslucentBlendMode(Material->GetBlendMode()))
+					)
 				{
 					continue;
 				}
@@ -746,19 +774,21 @@ void FEdModeFoliage::RemoveInstancesForBrush(AInstancedFoliageActor* IFA, FFolia
 		}
 	}
 
-	if (!UISettings.bFilterLandscape || !UISettings.bFilterStaticMesh || !UISettings.bFilterBSP)
+	if (!UISettings.bFilterLandscape || !UISettings.bFilterStaticMesh || !UISettings.bFilterBSP || !UISettings.bFilterTranslucent)
 	{
 		// Filter ExistingInstances
 		for (int32 Idx = 0; Idx < ExistingInstances.Num(); Idx++)
 		{
-			UActorComponent* Base = MeshInfo.Instances[ExistingInstances[Idx]].Base;
+			UPrimitiveComponent* Base = MeshInfo.Instances[ExistingInstances[Idx]].Base;
+			UMaterialInterface* Material = Base ? Base->GetMaterial(0) : nullptr;
 
 			// Check if instance is candidate for removal based on filter settings
-			if ((Base && (
+			if (Base && (
 				(!UISettings.bFilterLandscape && Base->IsA(ULandscapeHeightfieldCollisionComponent::StaticClass())) ||
 				(!UISettings.bFilterStaticMesh && Base->IsA(UStaticMeshComponent::StaticClass())) ||
-				(!UISettings.bFilterBSP && Base->IsA(UModelComponent::StaticClass()))
-				)))
+				(!UISettings.bFilterBSP && Base->IsA(UModelComponent::StaticClass())) ||
+				(!UISettings.bFilterTranslucent && Material && IsTranslucentBlendMode(Material->GetBlendMode()))
+				))
 			{
 				// Instance should not be removed, so remove it from the removal list.
 				ExistingInstances.RemoveSwap(Idx);
@@ -959,12 +989,7 @@ void FEdModeFoliage::ReapplyInstancesForBrush(UWorld* InWorld, AInstancedFoliage
 				FVector ZAxis = Instance.Rotation.Quaternion().GetAxisZ();
 				FVector Start = Instance.Location + 16.f * ZAxis;
 				FVector End = Instance.Location - 16.f * ZAxis;
-				FCollisionQueryParams TraceParams(NAME_ReapplyInstancesForBrush, true);
-				if (Settings->ReapplyVertexColorMask && Settings->VertexColorMask != FOLIAGEVERTEXCOLORMASK_Disabled)
-				{
-					TraceParams.bReturnFaceIndex = true;
-				}
-				if (InWorld->LineTraceSingle(Hit, Start, End, TraceParams, FCollisionObjectQueryParams(ECC_WorldStatic)))
+				if (FoliageTrace(InWorld, Hit, Start, End, NAME_ReapplyInstancesForBrush, true))
 				{
 					// Reapply the normal
 					if (bReapplyNormal)
@@ -1070,7 +1095,7 @@ void FEdModeFoliage::ReapplyInstancesForBrush(UWorld* InWorld, AInstancedFoliage
 				static const FName NAME_ReapplyInstancesForBrush = TEXT("ReapplyCollisionWithWorld");
 				FVector Start = Instance.Location + FVector(0.f, 0.f, 16.f);
 				FVector End = Instance.Location - FVector(0.f, 0.f, 16.f);
-				if (InWorld->LineTraceSingle(Hit, Start, End, FCollisionQueryParams(NAME_ReapplyInstancesForBrush, true), FCollisionObjectQueryParams(ECC_WorldStatic)))
+				if (FoliageTrace(InWorld, Hit, Start, End, NAME_ReapplyInstancesForBrush))
 				{
 					if (!CheckCollisionWithWorld(Settings, Instance, Hit.Normal, Hit.Location, InWorld))
 					{
@@ -1289,7 +1314,10 @@ void FEdModeFoliage::ApplyPaintBucket(AActor* Actor, bool bRemove)
 		for (int32 ComponentIdx = 0; ComponentIdx < StaticMeshComponents.Num(); ComponentIdx++)
 		{
 			UStaticMeshComponent* StaticMeshComponent = StaticMeshComponents[ComponentIdx];
-			if (UISettings.bFilterStaticMesh && StaticMeshComponent->StaticMesh && StaticMeshComponent->StaticMesh->RenderData)
+			UMaterialInterface* Material = StaticMeshComponent->GetMaterial(0);
+
+			if (UISettings.bFilterStaticMesh && StaticMeshComponent->StaticMesh && StaticMeshComponent->StaticMesh->RenderData &&
+				(UISettings.bFilterTranslucent || !Material || !IsTranslucentBlendMode(Material->GetBlendMode())))
 			{
 				UStaticMesh* StaticMesh = StaticMeshComponent->StaticMesh;
 				FStaticMeshLODResources& LODModel = StaticMesh->RenderData->LODResources[0];
@@ -1855,7 +1883,7 @@ bool FEdModeFoliage::InputKey(FEditorViewportClient* ViewportClient, FViewport*
 						FVector End = Instance.Location - FVector(0.f, 0.f, FOLIAGE_SNAP_TRACE);
 
 						FHitResult Hit;
-						if (World->LineTraceSingle(Hit, Start, End, FCollisionQueryParams(FName("FoliageSnap"), true), FCollisionObjectQueryParams(ECC_WorldStatic)))
+						if (FoliageTrace(GetWorld(), Hit, Start, End, FName("FoliageSnap")))
 						{
 							// Check current level
 							if ((Hit.Component.IsValid() && Hit.Component.Get()->GetOutermost() == World->GetCurrentLevel()->GetOutermost()) ||
@@ -2133,6 +2161,7 @@ void FFoliageUISettings::Load()
 	GConfig->GetBool(TEXT("FoliageEdit"), TEXT("bFilterLandscape"), bFilterLandscape, GEditorUserSettingsIni);
 	GConfig->GetBool(TEXT("FoliageEdit"), TEXT("bFilterStaticMesh"), bFilterStaticMesh, GEditorUserSettingsIni);
 	GConfig->GetBool(TEXT("FoliageEdit"), TEXT("bFilterBSP"), bFilterBSP, GEditorUserSettingsIni);
+	GConfig->GetBool(TEXT("FoliageEdit"), TEXT("bFilterTranslucent"), bFilterTranslucent, GEditorUserSettingsIni);
 }
 
 /** Save UI settings to ini file */
@@ -2147,4 +2176,5 @@ void FFoliageUISettings::Save()
 	GConfig->SetBool(TEXT("FoliageEdit"), TEXT("bFilterLandscape"), bFilterLandscape, GEditorUserSettingsIni);
 	GConfig->SetBool(TEXT("FoliageEdit"), TEXT("bFilterStaticMesh"), bFilterStaticMesh, GEditorUserSettingsIni);
 	GConfig->SetBool(TEXT("FoliageEdit"), TEXT("bFilterBSP"), bFilterBSP, GEditorUserSettingsIni);
+	GConfig->SetBool(TEXT("FoliageEdit"), TEXT("bFilterTranslucent"), bFilterTranslucent, GEditorUserSettingsIni);
 }
diff --git a/Engine/Source/Editor/FoliageEdit/Private/FoliageEdMode.h b/Engine/Source/Editor/FoliageEdit/Private/FoliageEdMode.h
--- a/Engine/Source/Editor/FoliageEdit/Private/FoliageEdMode.h
+++ b/Engine/Source/Editor/FoliageEdit/Private/FoliageEdMode.h
@@ -52,6 +52,8 @@ struct FFoliageUISettings
 	void SetFilterStaticMesh(bool InbFilterStaticMesh) { bFilterStaticMesh = InbFilterStaticMesh; }
 	bool GetFilterBSP() const { return bFilterBSP ? true : false; }
 	void SetFilterBSP(bool InbFilterBSP) { bFilterBSP = InbFilterBSP; }
+	bool GetFilterTranslucent() const { return bFilterTranslucent; }
+	void SetFilterTranslucent(bool InbFilterTranslucent) { bFilterTranslucent = InbFilterTranslucent; }
 
 	FFoliageUISettings()
 		: WindowX(-1)
@@ -98,6 +100,7 @@ public:
 	bool bFilterLandscape;
 	bool bFilterStaticMesh;
 	bool bFilterBSP;
+	bool bFilterTranslucent;
 };
 
 
diff --git a/Engine/Source/Editor/FoliageEdit/Private/SFoliageEdit.cpp b/Engine/Source/Editor/FoliageEdit/Private/SFoliageEdit.cpp
--- a/Engine/Source/Editor/FoliageEdit/Private/SFoliageEdit.cpp
+++ b/Engine/Source/Editor/FoliageEdit/Private/SFoliageEdit.cpp
@@ -206,7 +206,7 @@ void SFoliageEdit::Construct(const FArguments& InArgs)
 				.UseAllottedWidth(true)
 
 				+ SWrapBox::Slot()
-				.Padding(0.0f, 0.0f, 16.0f, 0.0f)
+				.Padding(0, 0, 16, 0)
 				[
 					SNew(SCheckBox)
 					.Visibility(this, &SFoliageEdit::GetVisibility_Filters)
@@ -219,7 +219,7 @@ void SFoliageEdit::Construct(const FArguments& InArgs)
 				]
 
 				+ SWrapBox::Slot()
-				.Padding(0.0f, 0.0f, 16.0f, 0.0f)
+				.Padding(0, 0, 16, 0)
 				[
 					SNew(SCheckBox)
 					.Visibility(this, &SFoliageEdit::GetVisibility_Filters)
@@ -232,6 +232,7 @@ void SFoliageEdit::Construct(const FArguments& InArgs)
 				]
 
 				+ SWrapBox::Slot()
+				.Padding(0, 0, 16, 0)
 				[
 					SNew(SCheckBox)
 					.Visibility(this, &SFoliageEdit::GetVisibility_Filters)
@@ -242,6 +243,19 @@ void SFoliageEdit::Construct(const FArguments& InArgs)
 						.Text(LOCTEXT("BSP", "BSP"))
 					]
 				]
+
+				+ SWrapBox::Slot()
+				.Padding(0, 0, 16, 0)
+				[
+					SNew(SCheckBox)
+					.Visibility(this, &SFoliageEdit::GetVisibility_Filters)
+					.OnCheckStateChanged(this, &SFoliageEdit::OnCheckStateChanged_Translucent)
+					.IsChecked(this, &SFoliageEdit::GetCheckState_Translucent)
+					[
+						SNew(STextBlock)
+						.Text(LOCTEXT("Translucent", "Translucent"))
+					]
+				]
 			]
 
 			+ SVerticalBox::Slot()
@@ -362,9 +376,9 @@ void SFoliageEdit::AddItemToScrollbox(struct FFoliageMeshUIInfo& InFoliageInfoTo
 	DisplayItemList.Add(DisplayItem);
 
 	ItemScrollBox->AddSlot()
-		[
-			DisplayItem
-		];
+	[
+		DisplayItem
+	];
 }
 
 void SFoliageEdit::RemoveItemFromScrollbox(const TSharedPtr<SFoliageEditMeshDisplayItem> InWidgetToRemove)
@@ -677,6 +691,16 @@ ESlateCheckBoxState::Type SFoliageEdit::GetCheckState_BSP() const
 	return FoliageEditMode->UISettings.GetFilterBSP() ? ESlateCheckBoxState::Checked : ESlateCheckBoxState::Unchecked;
 }
 
+void SFoliageEdit::OnCheckStateChanged_Translucent(ESlateCheckBoxState::Type InState)
+{
+	FoliageEditMode->UISettings.SetFilterTranslucent(InState == ESlateCheckBoxState::Checked ? true : false);
+}
+
+ESlateCheckBoxState::Type SFoliageEdit::GetCheckState_Translucent() const
+{
+	return FoliageEditMode->UISettings.GetFilterTranslucent() ? ESlateCheckBoxState::Checked : ESlateCheckBoxState::Unchecked;
+}
+
 EVisibility SFoliageEdit::GetVisibility_Radius() const
 {
 	if (FoliageEditMode->UISettings.GetSelectToolSelected() || FoliageEditMode->UISettings.GetReapplyPaintBucketToolSelected() || FoliageEditMode->UISettings.GetPaintBucketToolSelected())
diff --git a/Engine/Source/Editor/FoliageEdit/Private/SFoliageEdit.h b/Engine/Source/Editor/FoliageEdit/Private/SFoliageEdit.h
--- a/Engine/Source/Editor/FoliageEdit/Private/SFoliageEdit.h
+++ b/Engine/Source/Editor/FoliageEdit/Private/SFoliageEdit.h
@@ -141,6 +141,12 @@ private:
 	/** Retrieves the filter settings for painting on BSPs. */
 	ESlateCheckBoxState::Type GetCheckState_BSP() const;
 
+	/** Sets the filter settings for if painting will occur on translucent meshes. */
+	void OnCheckStateChanged_Translucent(ESlateCheckBoxState::Type InState);
+
+	/** Retrieves the filter settings for painting on translucent meshes. */
+	ESlateCheckBoxState::Type GetCheckState_Translucent() const;
+
 	/** Checks if the text in the empty list overlay should appear. If the list is has items but the the drag and drop override is true, it will return EVisibility::Visible. */
 	EVisibility GetVisibility_EmptyList() const;
 
diff --git a/Engine/Source/Editor/UnrealEd/Private/EditorObject.cpp b/Engine/Source/Editor/UnrealEd/Private/EditorObject.cpp
--- a/Engine/Source/Editor/UnrealEd/Private/EditorObject.cpp
+++ b/Engine/Source/Editor/UnrealEd/Private/EditorObject.cpp
@@ -240,7 +240,7 @@ static const TCHAR* ImportProperties(
 				ParseObject<UStaticMesh>(Str, TEXT("StaticMesh="), StaticMesh, ANY_PACKAGE) &&
 				FParse::Value(Str, TEXT("Component="), ComponentName) )
 			{
-				UActorComponent* ActorComponent = FindObjectFast<UActorComponent>(SubobjectRoot, ComponentName);
+				UPrimitiveComponent* ActorComponent = FindObjectFast<UPrimitiveComponent>(SubobjectRoot, ComponentName);
 
 				if (ActorComponent)
 				{
diff --git a/Engine/Source/Runtime/Engine/Classes/Foliage/InstancedFoliageActor.h b/Engine/Source/Runtime/Engine/Classes/Foliage/InstancedFoliageActor.h
--- a/Engine/Source/Runtime/Engine/Classes/Foliage/InstancedFoliageActor.h
+++ b/Engine/Source/Runtime/Engine/Classes/Foliage/InstancedFoliageActor.h
@@ -51,7 +51,7 @@ public:
 
 	// Change all instances based on one component to a new component (possible in another level).
 	// The instances keep the same world locations
-	ENGINE_API void MoveInstancesToNewComponent(UActorComponent* InOldComponent, UActorComponent* InNewComponent);
+	ENGINE_API void MoveInstancesToNewComponent(UPrimitiveComponent* InOldComponent, UPrimitiveComponent* InNewComponent);
 
 	// Move instances based on a component that has just been moved.
 	void MoveInstancesForMovedComponent(UActorComponent* InComponent);
diff --git a/Engine/Source/Runtime/Engine/Private/InstancedFoliage.cpp b/Engine/Source/Runtime/Engine/Private/InstancedFoliage.cpp
--- a/Engine/Source/Runtime/Engine/Private/InstancedFoliage.cpp
+++ b/Engine/Source/Runtime/Engine/Private/InstancedFoliage.cpp
@@ -974,7 +974,7 @@ void AInstancedFoliageActor::MoveInstancesForComponentToCurrentLevel(UActorCompo
 	}
 }
 
-void AInstancedFoliageActor::MoveInstancesToNewComponent(UActorComponent* InOldComponent, UActorComponent* InNewComponent)
+void AInstancedFoliageActor::MoveInstancesToNewComponent(UPrimitiveComponent* InOldComponent, UPrimitiveComponent* InNewComponent)
 {
 	AInstancedFoliageActor* NewIFA = AInstancedFoliageActor::GetInstancedFoliageActorForLevel(InNewComponent->GetTypedOuter<ULevel>());
 	if (!NewIFA)
@@ -1435,7 +1435,7 @@ void AInstancedFoliageActor::MapRebuild()
 
 					if (bHit && Result.Component.IsValid() && Result.Component->IsA(UModelComponent::StaticClass()))
 					{
-						NewInstance.Base = Cast<UActorComponent>(Result.Component.Get());
+						NewInstance.Base = CastChecked<UPrimitiveComponent>(Result.Component.Get());
 					}
 					else
 					{
diff --git a/Engine/Source/Runtime/Engine/Public/InstancedFoliage.h b/Engine/Source/Runtime/Engine/Public/InstancedFoliage.h
--- a/Engine/Source/Runtime/Engine/Public/InstancedFoliage.h
+++ b/Engine/Source/Runtime/Engine/Public/InstancedFoliage.h
@@ -51,7 +51,7 @@ struct FFoliageInstancePlacementInfo
  */
 struct FFoliageInstance : public FFoliageInstancePlacementInfo
 {
-	UActorComponent* Base;
+	UPrimitiveComponent* Base;
 	int32 ClusterIndex;	// -1/INDEX_NONE if this instance is invalid and its array slot can be reused.
 
 	FFoliageInstance()
@@ -85,11 +85,10 @@ struct FFoliageInstance : public FFoliageInstancePlacementInfo
 			{
 				AlignRotation.Pitch = MaxPitch;
 			}
-			else
-				if (AlignRotation.Pitch < -MaxPitch)
-				{
+			else if (AlignRotation.Pitch < -MaxPitch)
+			{
 				AlignRotation.Pitch = -MaxPitch;
-				}
+			}
 		}
 
 		PreAlignRotation = Rotation;
