Fix foliage "fill" tool on spline meshes

diff --git a/Engine/Source/Editor/FoliageEdit/Private/FoliageEdMode.cpp b/Engine/Source/Editor/FoliageEdit/Private/FoliageEdMode.cpp
--- a/Engine/Source/Editor/FoliageEdit/Private/FoliageEdMode.cpp
+++ b/Engine/Source/Editor/FoliageEdit/Private/FoliageEdMode.cpp
@@ -27,6 +27,7 @@
 #include "Landscape/LandscapeComponent.h"
 #include "Landscape/LandscapeHeightfieldCollisionComponent.h"
 #include "Landscape/LandscapeInfo.h"
+#include "Components/SplineMeshComponent.h"
 
 #define FOLIAGE_SNAP_TRACE (10000.f)
 
@@ -1305,8 +1306,6 @@ void FEdModeFoliage::ApplyPaintBucket(AActor* Actor, bool bRemove)
 	{
 		TMap<UPrimitiveComponent*, TArray<FFoliagePaintBucketTriangle> > ComponentPotentialTriangles;
 
-		FTransform LocalToWorld = Actor->GetRootComponent()->ComponentToWorld;
-
 		// Check all the components of the hit actor
 		TArray<UStaticMeshComponent*> StaticMeshComponents;
 		Actor->GetComponents(StaticMeshComponents);
@@ -1334,25 +1333,55 @@ void FEdModeFoliage::ApplyPaintBucket(AActor* Actor, bool bRemove)
 				const bool bHasColorData = bHasInstancedColorData || LODModel.ColorVertexBuffer.GetNumVertices();
 
 				// Get the raw triangle data for this static mesh
+				FTransform LocalToWorld = StaticMeshComponent->ComponentToWorld;
 				FIndexArrayView Indices = LODModel.IndexBuffer.GetArrayView();
 				const FPositionVertexBuffer& PositionVertexBuffer = LODModel.PositionVertexBuffer;
 				const FColorVertexBuffer& ColorVertexBuffer = LODModel.ColorVertexBuffer;
 
-				// Build a mapping of vertex positions to vertex colors.  Using a TMap will allow for fast lookups so we can match new static mesh vertices with existing colors 
-				for (int32 Idx = 0; Idx < Indices.Num(); Idx += 3)
+				if (USplineMeshComponent* SplineMesh = Cast<USplineMeshComponent>(StaticMeshComponent))
 				{
-					const int32 Index0 = Indices[Idx];
-					const int32 Index1 = Indices[Idx + 1];
-					const int32 Index2 = Indices[Idx + 2];
-
-					new(PotentialTriangles)FFoliagePaintBucketTriangle(LocalToWorld
-						, PositionVertexBuffer.VertexPosition(Index0)
-						, PositionVertexBuffer.VertexPosition(Index1)
-						, PositionVertexBuffer.VertexPosition(Index2)
-						, bHasInstancedColorData ? InstanceMeshLODInfo->PaintedVertices[Index0].Color : (bHasColorData ? ColorVertexBuffer.VertexColor(Index0) : FColor::White)
-						, bHasInstancedColorData ? InstanceMeshLODInfo->PaintedVertices[Index1].Color : (bHasColorData ? ColorVertexBuffer.VertexColor(Index1) : FColor::White)
-						, bHasInstancedColorData ? InstanceMeshLODInfo->PaintedVertices[Index2].Color : (bHasColorData ? ColorVertexBuffer.VertexColor(Index2) : FColor::White)
-						);
+					// Transform spline mesh verts correctly
+					FVector Mask = FVector(1, 1, 1);
+					USplineMeshComponent::GetAxisValue(Mask, SplineMesh->ForwardAxis) = 0;
+
+					for (int32 Idx = 0; Idx < Indices.Num(); Idx += 3)
+					{
+						const int32 Index0 = Indices[Idx];
+						const int32 Index1 = Indices[Idx + 1];
+						const int32 Index2 = Indices[Idx + 2];
+
+						const FVector Vert0 = SplineMesh->CalcSliceTransform(USplineMeshComponent::GetAxisValue(PositionVertexBuffer.VertexPosition(Index0), SplineMesh->ForwardAxis)).TransformPosition(PositionVertexBuffer.VertexPosition(Index0) * Mask);
+						const FVector Vert1 = SplineMesh->CalcSliceTransform(USplineMeshComponent::GetAxisValue(PositionVertexBuffer.VertexPosition(Index1), SplineMesh->ForwardAxis)).TransformPosition(PositionVertexBuffer.VertexPosition(Index1) * Mask);
+						const FVector Vert2 = SplineMesh->CalcSliceTransform(USplineMeshComponent::GetAxisValue(PositionVertexBuffer.VertexPosition(Index2), SplineMesh->ForwardAxis)).TransformPosition(PositionVertexBuffer.VertexPosition(Index2) * Mask);
+
+						new(PotentialTriangles)FFoliagePaintBucketTriangle(LocalToWorld
+							, Vert0
+							, Vert1
+							, Vert2
+							, bHasInstancedColorData ? InstanceMeshLODInfo->PaintedVertices[Index0].Color : (bHasColorData ? ColorVertexBuffer.VertexColor(Index0) : FColor::White)
+							, bHasInstancedColorData ? InstanceMeshLODInfo->PaintedVertices[Index1].Color : (bHasColorData ? ColorVertexBuffer.VertexColor(Index1) : FColor::White)
+							, bHasInstancedColorData ? InstanceMeshLODInfo->PaintedVertices[Index2].Color : (bHasColorData ? ColorVertexBuffer.VertexColor(Index2) : FColor::White)
+							);
+					}
+				}
+				else
+				{
+					// Build a mapping of vertex positions to vertex colors.  Using a TMap will allow for fast lookups so we can match new static mesh vertices with existing colors 
+					for (int32 Idx = 0; Idx < Indices.Num(); Idx += 3)
+					{
+						const int32 Index0 = Indices[Idx];
+						const int32 Index1 = Indices[Idx + 1];
+						const int32 Index2 = Indices[Idx + 2];
+
+						new(PotentialTriangles)FFoliagePaintBucketTriangle(LocalToWorld
+							, PositionVertexBuffer.VertexPosition(Index0)
+							, PositionVertexBuffer.VertexPosition(Index1)
+							, PositionVertexBuffer.VertexPosition(Index2)
+							, bHasInstancedColorData ? InstanceMeshLODInfo->PaintedVertices[Index0].Color : (bHasColorData ? ColorVertexBuffer.VertexColor(Index0) : FColor::White)
+							, bHasInstancedColorData ? InstanceMeshLODInfo->PaintedVertices[Index1].Color : (bHasColorData ? ColorVertexBuffer.VertexColor(Index1) : FColor::White)
+							, bHasInstancedColorData ? InstanceMeshLODInfo->PaintedVertices[Index2].Color : (bHasColorData ? ColorVertexBuffer.VertexColor(Index2) : FColor::White)
+							);
+					}
 				}
 			}
 		}
diff --git a/Engine/Source/Editor/UnrealEd/Private/Fbx/FbxMainExport.cpp b/Engine/Source/Editor/UnrealEd/Private/Fbx/FbxMainExport.cpp
--- a/Engine/Source/Editor/UnrealEd/Private/Fbx/FbxMainExport.cpp
+++ b/Engine/Source/Editor/UnrealEd/Private/Fbx/FbxMainExport.cpp
@@ -2015,22 +2015,6 @@ FbxNode* FFbxExporter::ExportStaticMeshToFbx(UStaticMesh* StaticMesh, FStaticMes
 	return FbxActor;
 }
 
-static float& GetAxisValue(FVector& InVector, ESplineMeshAxis::Type InAxis)
-{
-	switch (InAxis)
-	{
-	case ESplineMeshAxis::X:
-		return InVector.X;
-	case ESplineMeshAxis::Y:
-		return InVector.Y;
-	case ESplineMeshAxis::Z:
-		return InVector.Z;
-	default:
-		check(0);
-		return InVector.Z;
-	}
-}
-
 FbxNode* FFbxExporter::ExportSplineMeshToFbx(USplineMeshComponent* SplineMeshComp, FStaticMeshLODResources& RenderMesh, const TCHAR* MeshName, FbxNode* FbxActor)
 {
 	const UStaticMesh* StaticMesh = SplineMeshComp->StaticMesh;
@@ -2082,8 +2066,8 @@ FbxNode* FFbxExporter::ExportSplineMeshToFbx(USplineMeshComponent* SplineMeshCom
 		int32 UnrealPosIndex = UniqueVerts[PosIndex];
 		FVector Position = RenderMesh.PositionVertexBuffer.VertexPosition(UnrealPosIndex);
 
-		const FTransform SliceTransform = SplineMeshComp->CalcSliceTransform(GetAxisValue(Position, SplineMeshComp->ForwardAxis));
-		GetAxisValue(Position, SplineMeshComp->ForwardAxis) = 0;
+		const FTransform SliceTransform = SplineMeshComp->CalcSliceTransform(USplineMeshComponent::GetAxisValue(Position, SplineMeshComp->ForwardAxis));
+		USplineMeshComponent::GetAxisValue(Position, SplineMeshComp->ForwardAxis) = 0;
 		Position = SliceTransform.TransformPosition(Position);
 
 		ControlPoints[PosIndex] = FbxVector4(Position.X, -Position.Y, Position.Z);
@@ -2129,7 +2113,7 @@ FbxNode* FFbxExporter::ExportSplineMeshToFbx(USplineMeshComponent* SplineMeshCom
 	for (int32 VertIndex = 0; VertIndex < VertexCount; ++VertIndex)
 	{
 		FVector Position = RenderMesh.PositionVertexBuffer.VertexPosition(VertIndex);
-		const FTransform SliceTransform = SplineMeshComp->CalcSliceTransform(GetAxisValue(Position, SplineMeshComp->ForwardAxis));
+		const FTransform SliceTransform = SplineMeshComp->CalcSliceTransform(USplineMeshComponent::GetAxisValue(Position, SplineMeshComp->ForwardAxis));
 		FVector Normal = FVector(RenderMesh.VertexBuffer.VertexTangentZ(VertIndex));
 		Normal = SliceTransform.TransformVector(Normal);
 		FbxVector4& FbxNormal = FbxNormals[VertIndex];
diff --git a/Engine/Source/Runtime/Engine/Classes/Components/SplineMeshComponent.h b/Engine/Source/Runtime/Engine/Classes/Components/SplineMeshComponent.h
--- a/Engine/Source/Runtime/Engine/Classes/Components/SplineMeshComponent.h
+++ b/Engine/Source/Runtime/Engine/Classes/Components/SplineMeshComponent.h
@@ -209,10 +209,39 @@ class ENGINE_API USplineMeshComponent : public UStaticMeshComponent, public IInt
 	 */
 	FTransform CalcSliceTransform(const float DistanceAlong) const;
 
-	static const float& GetAxisValue(const FVector& InVector, ESplineMeshAxis::Type InAxis);
-	static float& GetAxisValue(FVector& InVector, ESplineMeshAxis::Type InAxis);
-
+	inline static const float& GetAxisValue(const FVector& InVector, ESplineMeshAxis::Type InAxis);
+	inline static float& GetAxisValue(FVector& InVector, ESplineMeshAxis::Type InAxis);
 };
 
+const float& USplineMeshComponent::GetAxisValue(const FVector& InVector, ESplineMeshAxis::Type InAxis)
+{
+	switch (InAxis)
+	{
+	case ESplineMeshAxis::X:
+		return InVector.X;
+	case ESplineMeshAxis::Y:
+		return InVector.Y;
+	case ESplineMeshAxis::Z:
+		return InVector.Z;
+	default:
+		check(0);
+		return InVector.Z;
+	}
+}
 
+float& USplineMeshComponent::GetAxisValue(FVector& InVector, ESplineMeshAxis::Type InAxis)
+{
+	switch (InAxis)
+	{
+	case ESplineMeshAxis::X:
+		return InVector.X;
+	case ESplineMeshAxis::Y:
+		return InVector.Y;
+	case ESplineMeshAxis::Z:
+		return InVector.Z;
+	default:
+		check(0);
+		return InVector.Z;
+	}
+}
 
diff --git a/Engine/Source/Runtime/Engine/Private/SplineMeshComponent.cpp b/Engine/Source/Runtime/Engine/Private/SplineMeshComponent.cpp
--- a/Engine/Source/Runtime/Engine/Private/SplineMeshComponent.cpp
+++ b/Engine/Source/Runtime/Engine/Private/SplineMeshComponent.cpp
@@ -740,36 +740,3 @@ FStaticMeshStaticLightingMesh* USplineMeshComponent::AllocateStaticLightingMesh(
 {
 	return new FSplineStaticLightingMesh(this, LODIndex, InRelevantLights);
 }
-
-
-const float& USplineMeshComponent::GetAxisValue(const FVector& InVector, ESplineMeshAxis::Type InAxis)
-{
-	switch (InAxis)
-	{
-	case ESplineMeshAxis::X:
-		return InVector.X;
-	case ESplineMeshAxis::Y:
-		return InVector.Y;
-	case ESplineMeshAxis::Z:
-		return InVector.Z;
-	default:
-		check(0);
-		return InVector.Z;
-	}
-}
-
-float& USplineMeshComponent::GetAxisValue(FVector& InVector, ESplineMeshAxis::Type InAxis)
-{
-	switch (InAxis)
-	{
-	case ESplineMeshAxis::X:
-		return InVector.X;
-	case ESplineMeshAxis::Y:
-		return InVector.Y;
-	case ESplineMeshAxis::Z:
-		return InVector.Z;
-	default:
-		check(0);
-		return InVector.Z;
-	}
-}
\ No newline at end of file
