Fix crash when using foliage reapply tool for landscape layer when layer name is set to "none"

diff --git a/Engine/Source/Editor/FoliageEdit/Private/FoliageEdMode.cpp b/Engine/Source/Editor/FoliageEdit/Private/FoliageEdMode.cpp
--- a/Engine/Source/Editor/FoliageEdit/Private/FoliageEdMode.cpp
+++ b/Engine/Source/Editor/FoliageEdit/Private/FoliageEdMode.cpp
@@ -1011,7 +1011,7 @@ void FEdModeFoliage::ReapplyInstancesForBrush(UWorld* InWorld, AInstancedFoliage
 					}
 
 					// Cull instances for the landscape layer
-					if (Settings->ReapplyLandscapeLayer)
+					if (Settings->ReapplyLandscapeLayer && LandscapeLayerName != NAME_None)
 					{
 						float HitWeight = 1.f;
 						ULandscapeHeightfieldCollisionComponent* HitLandscapeCollision = Cast<ULandscapeHeightfieldCollisionComponent>(Hit.Component.Get());
