Fix crash removing foliage entries for null staticmesh

diff --git a/Engine/Source/Editor/FoliageEdit/Private/FoliageEdMode.cpp b/Engine/Source/Editor/FoliageEdit/Private/FoliageEdMode.cpp
--- a/Engine/Source/Editor/FoliageEdit/Private/FoliageEdMode.cpp
+++ b/Engine/Source/Editor/FoliageEdit/Private/FoliageEdMode.cpp
@@ -1613,9 +1613,15 @@ bool FEdModeFoliage::RemoveFoliageMesh(UFoliageType* Settings)
 	if (MeshInfo != nullptr)
 	{
 		int32 InstancesNum = MeshInfo->Instances.Num() - MeshInfo->FreeInstanceIndices.Num();
-		FText Message = FText::Format(NSLOCTEXT("UnrealEd", "FoliageMode_DeleteMesh", "Are you sure you want to remove all {0} instances of {1} from this level?"), FText::AsNumber(InstancesNum), FText::FromName(Settings->GetStaticMesh()->GetFName()));
 
-		if (InstancesNum == 0 || EAppReturnType::Yes == FMessageDialog::Open(EAppMsgType::YesNo, Message))
+		bool bProceed = true;
+		if (Settings->GetStaticMesh() != nullptr && InstancesNum > 0)
+		{
+			FText Message = FText::Format(NSLOCTEXT("UnrealEd", "FoliageMode_DeleteMesh", "Are you sure you want to remove all {0} instances of {1} from this level?"), FText::AsNumber(InstancesNum), FText::FromName(Settings->GetStaticMesh()->GetFName()));
+			bProceed = (FMessageDialog::Open(EAppMsgType::YesNo, Message) == EAppReturnType::Yes);
+		}
+
+		if (bProceed)
 		{
 			GEditor->BeginTransaction(NSLOCTEXT("UnrealEd", "FoliageMode_RemoveMeshTransaction", "Foliage Editing: Remove Mesh"));
 			IFA->RemoveMesh(Settings);
