Fix crash in foliage edit mode display

diff --git a/Engine/Source/Editor/FoliageEdit/Private/SFoliageEditMeshDisplayItem.cpp b/Engine/Source/Editor/FoliageEdit/Private/SFoliageEditMeshDisplayItem.cpp
--- a/Engine/Source/Editor/FoliageEdit/Private/SFoliageEditMeshDisplayItem.cpp
+++ b/Engine/Source/Editor/FoliageEdit/Private/SFoliageEditMeshDisplayItem.cpp
@@ -1881,12 +1881,12 @@ bool SFoliageEditMeshDisplayItem::IsPropertyVisible(const FPropertyAndParent& Pr
 
 FText SFoliageEditMeshDisplayItem::GetInstanceCountString() const
 {
-	return FText::Format(LOCTEXT("InstanceCount_Value", "Instance Count: {0}"), FText::AsNumber(FoliageMeshUIInfo->MeshInfo->GetInstanceCount()));
+	return FText::Format(LOCTEXT("InstanceCount_Value", "Instance Count: {0}"), FText::AsNumber(FoliageMeshUIInfo->MeshInfo ? FoliageMeshUIInfo->MeshInfo->GetInstanceCount() : 0 ));
 }
 
 FText SFoliageEditMeshDisplayItem::GetInstanceClusterCountString() const
 {
-	return FText::Format(LOCTEXT("ClusterCount_Value", "Cluster Count: {0}"), FText::AsNumber(FoliageMeshUIInfo->MeshInfo->InstanceClusters.Num()));
+	return FText::Format(LOCTEXT("ClusterCount_Value", "Cluster Count: {0}"), FText::AsNumber(FoliageMeshUIInfo->MeshInfo ? FoliageMeshUIInfo->MeshInfo->InstanceClusters.Num() : 0));
 }
 
 EVisibility SFoliageEditMeshDisplayItem::IsReapplySettingsVisible() const
