Fix for user not being prompted to save a shared foliage settings asset

diff --git a/Engine/Source/Editor/FoliageEdit/Private/FoliageEdMode.cpp b/Engine/Source/Editor/FoliageEdit/Private/FoliageEdMode.cpp
--- a/Engine/Source/Editor/FoliageEdit/Private/FoliageEdMode.cpp
+++ b/Engine/Source/Editor/FoliageEdit/Private/FoliageEdMode.cpp
@@ -1731,6 +1731,7 @@ UFoliageType* FEdModeFoliage::SaveSettingsObject(const FText& InSettingsPackageN
 
 	UFoliageType* NewSettings = Cast<UFoliageType>(StaticDuplicateObject(Settings, Package, *FPackageName::GetLongPackageAssetName(PackageName)));
 	NewSettings->SetFlags(RF_Standalone | RF_Public);
+	NewSettings->Modify();
 
 	ReplaceSettingsObject(Settings, NewSettings);
 
