Fixed foliage tool data internal data corruption when removing instances, possibly leading to infinite loop.
Fixed foliage erase density and reapply density features

diff --git a/Engine/Source/Editor/FoliageEdit/Private/FoliageEdMode.cpp b/Engine/Source/Editor/FoliageEdit/Private/FoliageEdMode.cpp
--- a/Engine/Source/Editor/FoliageEdit/Private/FoliageEdMode.cpp
+++ b/Engine/Source/Editor/FoliageEdit/Private/FoliageEdMode.cpp
@@ -606,7 +606,7 @@ static bool CheckVertexColor(const UFoliageType* Settings, const FColor& VertexC
 
 
 /** Add instances inside the brush to match DesiredInstanceCount */
-void FEdModeFoliage::AddInstancesForBrush(UWorld* InWorld, AInstancedFoliageActor* IFA, UFoliageType* Settings, FFoliageMeshInfo& MeshInfo, int32 DesiredInstanceCount, TArray<int32>& ExistingInstances, float Pressure)
+void FEdModeFoliage::AddInstancesForBrush(UWorld* InWorld, AInstancedFoliageActor* IFA, UFoliageType* Settings, FFoliageMeshInfo& MeshInfo, int32 DesiredInstanceCount, const TArray<int32>& ExistingInstances, float Pressure)
 {
 	checkf(InWorld == IFA->GetWorld(), TEXT("Warning:World does not match Foliage world"));
 	if (DesiredInstanceCount > ExistingInstances.Num())
@@ -762,25 +762,25 @@ void FEdModeFoliage::AddInstancesForBrush(UWorld* InWorld, AInstancedFoliageActo
 }
 
 /** Remove instances inside the brush to match DesiredInstanceCount */
-void FEdModeFoliage::RemoveInstancesForBrush(AInstancedFoliageActor* IFA, FFoliageMeshInfo& MeshInfo, int32 DesiredInstanceCount, TArray<int32>& ExistingInstances, float Pressure)
+void FEdModeFoliage::RemoveInstancesForBrush(AInstancedFoliageActor* IFA, FFoliageMeshInfo& MeshInfo, int32 DesiredInstanceCount, TArray<int32>& PotentialInstancesToRemove, float Pressure)
 {
-	int32 InstancesToRemove = FMath::RoundToInt((float)(ExistingInstances.Num() - DesiredInstanceCount) * Pressure);
-	int32 InstancesToKeep = ExistingInstances.Num() - InstancesToRemove;
+	int32 InstancesToRemove = FMath::RoundToInt((float)(PotentialInstancesToRemove.Num() - DesiredInstanceCount) * Pressure);
+	int32 InstancesToKeep = PotentialInstancesToRemove.Num() - InstancesToRemove;
 	if (InstancesToKeep > 0)
 	{
-		// Remove InstancesToKeep random ExistingInstances from the array to leave those ExistingInstances behind, and delete all the rest
+		// Remove InstancesToKeep random PotentialInstancesToRemove from the array to leave those PotentialInstancesToRemove behind, and delete all the rest
 		for (int32 i = 0; i < InstancesToKeep; i++)
 		{
-			ExistingInstances.RemoveSwap(FMath::Rand() % ExistingInstances.Num());
+			PotentialInstancesToRemove.RemoveAtSwap(FMath::Rand() % PotentialInstancesToRemove.Num());
 		}
 	}
 
 	if (!UISettings.bFilterLandscape || !UISettings.bFilterStaticMesh || !UISettings.bFilterBSP || !UISettings.bFilterTranslucent)
 	{
-		// Filter ExistingInstances
-		for (int32 Idx = 0; Idx < ExistingInstances.Num(); Idx++)
+		// Filter PotentialInstancesToRemove
+		for (int32 Idx = 0; Idx < PotentialInstancesToRemove.Num(); Idx++)
 		{
-			UPrimitiveComponent* Base = MeshInfo.Instances[ExistingInstances[Idx]].Base;
+			UPrimitiveComponent* Base = MeshInfo.Instances[PotentialInstancesToRemove[Idx]].Base;
 			UMaterialInterface* Material = Base ? Base->GetMaterial(0) : nullptr;
 
 			// Check if instance is candidate for removal based on filter settings
@@ -792,21 +792,21 @@ void FEdModeFoliage::RemoveInstancesForBrush(AInstancedFoliageActor* IFA, FFolia
 				))
 			{
 				// Instance should not be removed, so remove it from the removal list.
-				ExistingInstances.RemoveSwap(Idx);
+				PotentialInstancesToRemove.RemoveAtSwap(Idx);
 				Idx--;
 			}
 		}
 	}
 
-	// Remove ExistingInstances to reduce it to desired count
-	if (ExistingInstances.Num() > 0)
+	// Remove PotentialInstancesToRemove to reduce it to desired count
+	if (PotentialInstancesToRemove.Num() > 0)
 	{
-		MeshInfo.RemoveInstances(IFA, ExistingInstances);
+		MeshInfo.RemoveInstances(IFA, PotentialInstancesToRemove);
 	}
 }
 
 /** Reapply instance settings to exiting instances */
-void FEdModeFoliage::ReapplyInstancesForBrush(UWorld* InWorld, AInstancedFoliageActor* IFA, UFoliageType* Settings, TArray<int32>& ExistingInstances)
+void FEdModeFoliage::ReapplyInstancesForBrush(UWorld* InWorld, AInstancedFoliageActor* IFA, UFoliageType* Settings, const TArray<int32>& ExistingInstances)
 {
 	checkf(InWorld == IFA->GetWorld(), TEXT("Warning:World does not match Foliage world"));
 	FFoliageMeshInfo* MeshInfo = IFA->FindMesh(Settings);
@@ -1164,8 +1164,8 @@ void FEdModeFoliage::ApplyBrush(FEditorViewportClient* ViewportClient)
 				MeshInfo.SelectInstances(IFA, !IsShiftDown(ViewportClient->Viewport), Instances);
 			}
 			else
-				if (UISettings.GetReapplyToolSelected())
-				{
+			if (UISettings.GetReapplyToolSelected())
+			{
 				if (Settings->ReapplyDensity)
 				{
 					// Adjust instance density
@@ -1187,30 +1187,30 @@ void FEdModeFoliage::ApplyBrush(FEditorViewportClient* ViewportClient)
 
 				// Reapply any settings checked by the user
 				ReapplyInstancesForBrush(World, IFA, Settings, Instances);
-				}
-				else if (UISettings.GetPaintToolSelected())
+			}
+			else if (UISettings.GetPaintToolSelected())
+			{
+				// Shift unpaints
+				if (IsShiftDown(ViewportClient->Viewport))
 				{
-					// Shift unpaints
-					if (IsShiftDown(ViewportClient->Viewport))
-					{
-						int32 DesiredInstanceCount = FMath::RoundToInt(BrushArea * Settings->Density * UISettings.GetUnpaintDensity() / (1000.f*1000.f));
+					int32 DesiredInstanceCount = FMath::RoundToInt(BrushArea * Settings->Density * UISettings.GetUnpaintDensity() / (1000.f*1000.f));
 
-						if (DesiredInstanceCount < Instances.Num())
-						{
-							RemoveInstancesForBrush(IFA, MeshInfo, DesiredInstanceCount, Instances, Pressure);
-						}
-					}
-					else
+					if (DesiredInstanceCount < Instances.Num())
 					{
-						// This is the total set of instances disregarding parameters like slope, height or layer.
-						float DesiredInstanceCountFloat = BrushArea * Settings->Density * UISettings.GetPaintDensity() / (1000.f*1000.f);
+						RemoveInstancesForBrush(IFA, MeshInfo, DesiredInstanceCount, Instances, Pressure);
+					}
+				}
+				else
+				{
+					// This is the total set of instances disregarding parameters like slope, height or layer.
+					float DesiredInstanceCountFloat = BrushArea * Settings->Density * UISettings.GetPaintDensity() / (1000.f*1000.f);
 
-						// Allow a single instance with a random chance, if the brush is smaller than the density
-						int32 DesiredInstanceCount = DesiredInstanceCountFloat > 1.f ? FMath::RoundToInt(DesiredInstanceCountFloat) : FMath::FRand() < DesiredInstanceCountFloat ? 1 : 0;
+					// Allow a single instance with a random chance, if the brush is smaller than the density
+					int32 DesiredInstanceCount = DesiredInstanceCountFloat > 1.f ? FMath::RoundToInt(DesiredInstanceCountFloat) : FMath::FRand() < DesiredInstanceCountFloat ? 1 : 0;
 
-						AddInstancesForBrush(World, IFA, Settings, MeshInfo, DesiredInstanceCount, Instances, Pressure);
-					}
+					AddInstancesForBrush(World, IFA, Settings, MeshInfo, DesiredInstanceCount, Instances, Pressure);
 				}
+			}
 		}
 	}
 	if (UISettings.GetLassoSelectToolSelected())
diff --git a/Engine/Source/Editor/FoliageEdit/Private/FoliageEdMode.h b/Engine/Source/Editor/FoliageEdit/Private/FoliageEdMode.h
--- a/Engine/Source/Editor/FoliageEdit/Private/FoliageEdMode.h
+++ b/Engine/Source/Editor/FoliageEdit/Private/FoliageEdMode.h
@@ -316,13 +316,13 @@ public:
 	void ApplyPaintBucket(AActor* Actor, bool bRemove);
 private:
 	/** Add instances inside the brush to match DesiredInstanceCount */
-	void AddInstancesForBrush(UWorld* InWorld, AInstancedFoliageActor* IFA, UFoliageType* Settings, FFoliageMeshInfo& MeshInfo, int32 DesiredInstanceCount, TArray<int32>& ExistingInstances, float Pressure);
+	void AddInstancesForBrush(UWorld* InWorld, AInstancedFoliageActor* IFA, UFoliageType* Settings, FFoliageMeshInfo& MeshInfo, int32 DesiredInstanceCount, const TArray<int32>& ExistingInstances, float Pressure);
 
-	/** Remove instances inside the brush to match DesiredInstanceCount */
-	void RemoveInstancesForBrush(AInstancedFoliageActor* IFA, FFoliageMeshInfo& MeshInfo, int32 DesiredInstanceCount, TArray<int32>& ExistingInstances, float Pressure);
+	/** Remove instances inside the brush to match DesiredInstanceCount. NOTE: PotentialInstancesToRemove array is modified by this function. */
+	void RemoveInstancesForBrush(AInstancedFoliageActor* IFA, FFoliageMeshInfo& MeshInfo, int32 DesiredInstanceCount, TArray<int32>& PotentialInstancesToRemove, float Pressure);
 
 	/** Reapply instance settings to exiting instances */
-	void ReapplyInstancesForBrush(UWorld* InWorld, AInstancedFoliageActor* IFA, UFoliageType* Settings, TArray<int32>& ExistingInstances);
+	void ReapplyInstancesForBrush(UWorld* InWorld, AInstancedFoliageActor* IFA, UFoliageType* Settings, const TArray<int32>& ExistingInstances);
 
 	/** Lookup the vertex color corresponding to a location traced on a static mesh */
 	bool GetStaticMeshVertexColorForHit(UStaticMeshComponent* InStaticMeshComponent, int32 InTriangleIndex, const FVector& InHitLocation, FColor& OutVertexColor) const;
diff --git a/Engine/Source/Editor/FoliageEdit/Private/SFoliageEditMeshDisplayItem.cpp b/Engine/Source/Editor/FoliageEdit/Private/SFoliageEditMeshDisplayItem.cpp
--- a/Engine/Source/Editor/FoliageEdit/Private/SFoliageEditMeshDisplayItem.cpp
+++ b/Engine/Source/Editor/FoliageEdit/Private/SFoliageEditMeshDisplayItem.cpp
@@ -152,8 +152,8 @@ void SFoliageEditMeshDisplayItem::Construct(const FArguments& InArgs)
 					.MinValue(0.0f)
 					.MaxValue(4.0f)
 					.MaxSliderValue(2.0f)
-					.Value(this, &SFoliageEditMeshDisplayItem::GetDensity)
-					.OnValueChanged(this, &SFoliageEditMeshDisplayItem::OnDensityChanged)
+					.Value(this, &SFoliageEditMeshDisplayItem::GetDensityReapply)
+					.OnValueChanged(this, &SFoliageEditMeshDisplayItem::OnDensityReapplyChanged)
 				]
 			]
 
