Support painting landscape and foliage without control held down

diff --git a/Engine/Source/Editor/FoliageEdit/Private/FoliageEdMode.cpp b/Engine/Source/Editor/FoliageEdit/Private/FoliageEdMode.cpp
--- a/Engine/Source/Editor/FoliageEdit/Private/FoliageEdMode.cpp
+++ b/Engine/Source/Editor/FoliageEdit/Private/FoliageEdMode.cpp
@@ -242,7 +242,7 @@ bool FEdModeFoliage::DisallowMouseDeltaTracking() const
 /** FEdMode: Called once per frame */
 void FEdModeFoliage::Tick(FEditorViewportClient* ViewportClient, float DeltaTime)
 {
-	if (IsCtrlDown(ViewportClient->Viewport) && bToolActive)
+	if (bToolActive)
 	{
 		ApplyBrush(ViewportClient);
 	}
@@ -310,41 +310,44 @@ static bool FoliageTrace(UWorld* InWorld, FHitResult& OutHit, FVector InStart, F
 void FEdModeFoliage::FoliageBrushTrace(FEditorViewportClient* ViewportClient, int32 MouseX, int32 MouseY)
 {
 	bBrushTraceValid = false;
-	if (UISettings.GetPaintToolSelected() || UISettings.GetReapplyToolSelected() || UISettings.GetLassoSelectToolSelected())
+	if (!ViewportClient->IsMovingCamera())
 	{
-		// Compute a world space ray from the screen space mouse coordinates
-		FSceneViewFamilyContext ViewFamily(FSceneViewFamily::ConstructionValues(
-			ViewportClient->Viewport,
-			ViewportClient->GetScene(),
-			ViewportClient->EngineShowFlags)
-			.SetRealtimeUpdate(ViewportClient->IsRealtime()));
-		FSceneView* View = ViewportClient->CalcSceneView(&ViewFamily);
-		FViewportCursorLocation MouseViewportRay(View, ViewportClient, MouseX, MouseY);
-
-		FVector Start = MouseViewportRay.GetOrigin();
-		BrushTraceDirection = MouseViewportRay.GetDirection();
-		FVector End = Start + WORLD_MAX * BrushTraceDirection;
-
-		FHitResult Hit;
-		UWorld* World = ViewportClient->GetWorld();
-		static FName NAME_FoliageBrush = FName(TEXT("FoliageBrush"));
-		if (FoliageTrace(World, Hit, Start, End, NAME_FoliageBrush))
+		if (UISettings.GetPaintToolSelected() || UISettings.GetReapplyToolSelected() || UISettings.GetLassoSelectToolSelected())
 		{
-			// Check filters
-			UPrimitiveComponent* PrimComp = Hit.Component.Get();
-			UMaterialInterface* Material = PrimComp ? PrimComp->GetMaterial(0) : nullptr;
-
-			if (PrimComp &&
-				PrimComp->GetOutermost() == World->GetCurrentLevel()->GetOutermost() &&
-				(UISettings.bFilterLandscape || !PrimComp->IsA(ULandscapeHeightfieldCollisionComponent::StaticClass())) &&
-				(UISettings.bFilterStaticMesh || !PrimComp->IsA(UStaticMeshComponent::StaticClass())) &&
-				(UISettings.bFilterBSP || !PrimComp->IsA(UModelComponent::StaticClass())) &&
-				(UISettings.bFilterTranslucent || !Material || !IsTranslucentBlendMode(Material->GetBlendMode()))
-				)
+			// Compute a world space ray from the screen space mouse coordinates
+			FSceneViewFamilyContext ViewFamily(FSceneViewFamily::ConstructionValues(
+				ViewportClient->Viewport,
+				ViewportClient->GetScene(),
+				ViewportClient->EngineShowFlags)
+				.SetRealtimeUpdate(ViewportClient->IsRealtime()));
+			FSceneView* View = ViewportClient->CalcSceneView(&ViewFamily);
+			FViewportCursorLocation MouseViewportRay(View, ViewportClient, MouseX, MouseY);
+
+			FVector Start = MouseViewportRay.GetOrigin();
+			BrushTraceDirection = MouseViewportRay.GetDirection();
+			FVector End = Start + WORLD_MAX * BrushTraceDirection;
+
+			FHitResult Hit;
+			UWorld* World = ViewportClient->GetWorld();
+			static FName NAME_FoliageBrush = FName(TEXT("FoliageBrush"));
+			if (FoliageTrace(World, Hit, Start, End, NAME_FoliageBrush))
 			{
-				// Adjust the sphere brush
-				BrushLocation = Hit.Location;
-				bBrushTraceValid = true;
+				// Check filters
+				UPrimitiveComponent* PrimComp = Hit.Component.Get();
+				UMaterialInterface* Material = PrimComp ? PrimComp->GetMaterial(0) : nullptr;
+
+				if (PrimComp &&
+					PrimComp->GetOutermost() == World->GetCurrentLevel()->GetOutermost() &&
+					(UISettings.bFilterLandscape || !PrimComp->IsA(ULandscapeHeightfieldCollisionComponent::StaticClass())) &&
+					(UISettings.bFilterStaticMesh || !PrimComp->IsA(UStaticMeshComponent::StaticClass())) &&
+					(UISettings.bFilterBSP || !PrimComp->IsA(UModelComponent::StaticClass())) &&
+					(UISettings.bFilterTranslucent || !Material || !IsTranslucentBlendMode(Material->GetBlendMode()))
+					)
+				{
+					// Adjust the sphere brush
+					BrushLocation = Hit.Location;
+					bBrushTraceValid = true;
+				}
 			}
 		}
 	}
@@ -1830,41 +1833,45 @@ bool FEdModeFoliage::InputKey(FEditorViewportClient* ViewportClient, FViewport*
 {
 	if (UISettings.GetPaintToolSelected() || UISettings.GetReapplyToolSelected() || UISettings.GetLassoSelectToolSelected())
 	{
-		if (Key == EKeys::LeftMouseButton && Event == IE_Pressed && IsCtrlDown(Viewport))
+		if (Key == EKeys::LeftMouseButton && Event == IE_Pressed)
 		{
-			if (!bToolActive)
+			if (!Viewport->KeyState(EKeys::MiddleMouseButton) && !Viewport->KeyState(EKeys::RightMouseButton)
+				&& !IsCtrlDown(Viewport) && !IsShiftDown(Viewport) && !IsAltDown(Viewport))
 			{
-				GEditor->BeginTransaction(NSLOCTEXT("UnrealEd", "FoliageMode_EditTransaction", "Foliage Editing"));
-				InstanceSnapshot.Empty();
-
-				// Special setup beginning a stroke with the Reapply tool
-				// Necessary so we don't keep reapplying settings over and over for the same instances.
-				if (UISettings.GetReapplyToolSelected())
+				if (!bToolActive)
 				{
-					AInstancedFoliageActor* IFA = AInstancedFoliageActor::GetInstancedFoliageActorForCurrentLevel(GetWorld());
-					for (auto& MeshPair : IFA->FoliageMeshes)
+					GEditor->BeginTransaction(NSLOCTEXT("UnrealEd", "FoliageMode_EditTransaction", "Foliage Editing"));
+					InstanceSnapshot.Empty();
+
+					// Special setup beginning a stroke with the Reapply tool
+					// Necessary so we don't keep reapplying settings over and over for the same instances.
+					if (UISettings.GetReapplyToolSelected())
 					{
-						FFoliageMeshInfo& MeshInfo = *MeshPair.Value;
-						if (MeshPair.Key->IsSelected)
+						AInstancedFoliageActor* IFA = AInstancedFoliageActor::GetInstancedFoliageActorForCurrentLevel(GetWorld());
+						for (auto& MeshPair : IFA->FoliageMeshes)
 						{
-							// Take a snapshot of all the locations
-							InstanceSnapshot.Add(MeshPair.Key, FMeshInfoSnapshot(&MeshInfo));
-
-							// Clear the "FOLIAGE_Readjusted" flag
-							for (int32 Idx = 0; Idx < MeshInfo.Instances.Num(); Idx++)
+							FFoliageMeshInfo& MeshInfo = *MeshPair.Value;
+							if (MeshPair.Key->IsSelected)
 							{
-								MeshInfo.Instances[Idx].Flags &= (~FOLIAGE_Readjusted);
+								// Take a snapshot of all the locations
+								InstanceSnapshot.Add(MeshPair.Key, FMeshInfoSnapshot(&MeshInfo));
+
+								// Clear the "FOLIAGE_Readjusted" flag
+								for (int32 Idx = 0; Idx < MeshInfo.Instances.Num(); Idx++)
+								{
+									MeshInfo.Instances[Idx].Flags &= (~FOLIAGE_Readjusted);
+								}
 							}
 						}
 					}
+					ApplyBrush(ViewportClient);
+					bToolActive = true;
+					return true;
 				}
 			}
-			ApplyBrush(ViewportClient);
-			bToolActive = true;
-			return true;
 		}
 
-		if (Event == IE_Released && bToolActive && (Key == EKeys::LeftMouseButton || Key == EKeys::LeftControl || Key == EKeys::RightControl))
+		if (bToolActive && Key == EKeys::LeftMouseButton && Event == IE_Released)
 		{
 			//Set the cursor position to that of the slate cursor so it wont snap back
 			Viewport->SetPreCaptureMousePosFromSlateCursor();
@@ -2049,10 +2056,9 @@ bool FEdModeFoliage::HandleClick(FEditorViewportClient* InViewportClient, HHitPr
 
 		return true;
 	}
-	else
-		if (UISettings.GetPaintBucketToolSelected() || UISettings.GetReapplyPaintBucketToolSelected())
-		{
-		if (HitProxy && HitProxy->IsA(HActor::StaticGetType()) && Click.IsControlDown())
+	else if (UISettings.GetPaintBucketToolSelected() || UISettings.GetReapplyPaintBucketToolSelected())
+	{
+		if (HitProxy && HitProxy->IsA(HActor::StaticGetType()))
 		{
 			GEditor->BeginTransaction(NSLOCTEXT("UnrealEd", "FoliageMode_EditTransaction", "Foliage Editing"));
 			ApplyPaintBucket(((HActor*)HitProxy)->Actor, Click.IsShiftDown());
@@ -2060,7 +2066,7 @@ bool FEdModeFoliage::HandleClick(FEditorViewportClient* InViewportClient, HHitPr
 		}
 
 		return true;
-		}
+	}
 
 	return FEdMode::HandleClick(InViewportClient, HitProxy, Click);
 }
diff --git a/Engine/Source/Editor/LandscapeEditor/Private/LandscapeEdMode.cpp b/Engine/Source/Editor/LandscapeEditor/Private/LandscapeEdMode.cpp
--- a/Engine/Source/Editor/LandscapeEditor/Private/LandscapeEdMode.cpp
+++ b/Engine/Source/Editor/LandscapeEditor/Private/LandscapeEdMode.cpp
@@ -1105,7 +1105,7 @@ bool FEdModeLandscape::HandleClick(FEditorViewportClient* InViewportClient, HHit
 	}
 
 	// Override Click Input for Splines Tool
-	if (CurrentTool && CurrentTool && CurrentTool->HandleClick(HitProxy, Click))
+	if (CurrentTool && CurrentTool->HandleClick(HitProxy, Click))
 	{
 		return true;
 	}
@@ -1164,25 +1164,16 @@ bool FEdModeLandscape::InputKey(FEditorViewportClient* ViewportClient, FViewport
 			return true;
 		}
 
-		if (CurrentTool && CurrentTool && CurrentTool->InputKey(ViewportClient, Viewport, Key, Event) == true)
+		if (CurrentTool && CurrentTool->InputKey(ViewportClient, Viewport, Key, Event) == true)
 		{
 			return true;
 		}
 
-		if ((Event == IE_Pressed || Event == IE_Released) && IsCtrlDown(Viewport))
+		if (Key == EKeys::LeftMouseButton && Event == IE_Pressed)
 		{
-			// Cheat, but works... :)
-			bool bNeedSelectGizmo = CurrentGizmoActor.IsValid() && (GLandscapeEditRenderMode & ELandscapeEditRenderMode::Gizmo) && CurrentGizmoActor->IsSelected();
-			GEditor->SelectNone(false, true);
-			if (bNeedSelectGizmo)
-			{
-				GEditor->SelectActor(CurrentGizmoActor.Get(), true, false);
-			}
-		}
-
-		if (Key == EKeys::LeftMouseButton)
-		{
-			if (Event == IE_Pressed && (IsCtrlDown(Viewport) || (Viewport->IsPenActive() && Viewport->GetTabletPressure() > 0.f)))
+			// Only activate tool if we're not already moving the camera and we're not trying to drag a transform widget
+			// Not using "if (!ViewportClient->IsMovingCamera())" because it's wrong in ortho viewports :D
+			if (!Viewport->KeyState(EKeys::MiddleMouseButton) && !Viewport->KeyState(EKeys::RightMouseButton) && !IsAltDown(Viewport) && ViewportClient->GetCurrentWidgetAxis() == EAxisList::None)
 			{
 				if (CurrentTool && (CurrentTool->GetSupportedTargetTypes() == ELandscapeToolTargetTypeMask::NA || CurrentToolTarget.TargetType != ELandscapeToolTargetType::Invalid))
 				{
@@ -1209,15 +1200,15 @@ bool FEdModeLandscape::InputKey(FEditorViewportClient* ViewportClient, FViewport
 			}
 		}
 
-		if (Key == EKeys::LeftMouseButton || Key == EKeys::LeftControl || Key == EKeys::RightControl)
+		if (Key == EKeys::LeftMouseButton)
 		{
-			if (Event == IE_Released && CurrentTool && CurrentTool && bToolActive)
+			if (Event == IE_Released && CurrentTool && bToolActive)
 			{
 				//Set the cursor position to that of the slate cursor so it wont snap back
 				Viewport->SetPreCaptureMousePosFromSlateCursor();
 				CurrentTool->EndTool(ViewportClient);
 				bToolActive = false;
-				return (Key == EKeys::LeftMouseButton);
+				return true;
 			}
 		}
 
@@ -1276,7 +1267,7 @@ bool FEdModeLandscape::InputKey(FEditorViewportClient* ViewportClient, FViewport
 		// Prev tool 
 		if (Event == IE_Pressed && Key == EKeys::Comma)
 		{
-			if (CurrentTool && CurrentTool && bToolActive)
+			if (CurrentTool && bToolActive)
 			{
 				CurrentTool->EndTool(ViewportClient);
 				bToolActive = false;
@@ -1292,7 +1283,7 @@ bool FEdModeLandscape::InputKey(FEditorViewportClient* ViewportClient, FViewport
 		// Next tool 
 		if (Event == IE_Pressed && Key == EKeys::Period)
 		{
-			if (CurrentTool && CurrentTool && bToolActive)
+			if (CurrentTool && bToolActive)
 			{
 				CurrentTool->EndTool(ViewportClient);
 				bToolActive = false;
@@ -1400,7 +1391,7 @@ bool FEdModeLandscape::InputDelta(FEditorViewportClient* InViewportClient, FView
 		}
 	}
 
-	if (CurrentTool && CurrentTool && CurrentTool->InputDelta(InViewportClient, InViewport, InDrag, InRot, InScale))
+	if (CurrentTool && CurrentTool->InputDelta(InViewportClient, InViewport, InDrag, InRot, InScale))
 	{
 		return true;
 	}
@@ -2096,7 +2087,7 @@ bool FEdModeLandscape::UsesTransformWidget() const
 	}
 
 	// Override Widget for Splines Tool
-	if (CurrentTool && CurrentTool && CurrentTool->UsesTransformWidget())
+	if (CurrentTool && CurrentTool->UsesTransformWidget())
 	{
 		return true;
 	}
@@ -2233,7 +2224,7 @@ bool FEdModeLandscape::IsSelectionAllowed(AActor* InActor, bool bInSelection) co
 	}
 
 	// Override Selection for Splines Tool
-	if (CurrentTool && CurrentTool && CurrentTool->OverrideSelection())
+	if (CurrentTool && CurrentTool->OverrideSelection())
 	{
 		return CurrentTool->IsSelectionAllowed(InActor, bInSelection);
 	}
diff --git a/Engine/Source/Editor/LandscapeEditor/Private/LandscapeEdModeSplineTools.cpp b/Engine/Source/Editor/LandscapeEditor/Private/LandscapeEdModeSplineTools.cpp
--- a/Engine/Source/Editor/LandscapeEditor/Private/LandscapeEdModeSplineTools.cpp
+++ b/Engine/Source/Editor/LandscapeEditor/Private/LandscapeEdModeSplineTools.cpp
@@ -619,26 +619,31 @@ public:
 
 	virtual bool BeginTool(FEditorViewportClient* ViewportClient, const FLandscapeToolTarget& InTarget, const FVector& InHitLocation) override
 	{
-		LandscapeInfo = InTarget.LandscapeInfo.Get();
-
-		ALandscapeProxy* Landscape = LandscapeInfo->GetCurrentLevelLandscapeProxy(true);
-		if (!Landscape)
+		if (ViewportClient->IsCtrlPressed())
 		{
-			return false;
-		}
+			LandscapeInfo = InTarget.LandscapeInfo.Get();
 
-		if (!Landscape->SplineComponent)
-		{
-			CreateSplineComponent(Landscape, FVector(1.f) / Landscape->GetRootComponent()->RelativeScale3D);
-		}
+			ALandscapeProxy* Landscape = LandscapeInfo->GetCurrentLevelLandscapeProxy(true);
+			if (!Landscape)
+			{
+				return false;
+			}
 
-		const FTransform LandscapeToSpline = Landscape->LandscapeActorToWorld().GetRelativeTransform(Landscape->SplineComponent->ComponentToWorld);
+			if (!Landscape->SplineComponent)
+			{
+				CreateSplineComponent(Landscape, FVector(1.f) / Landscape->GetRootComponent()->RelativeScale3D);
+			}
 
-		AddControlPoint(Landscape, LandscapeToSpline.TransformPosition(InHitLocation));
+			const FTransform LandscapeToSpline = Landscape->LandscapeActorToWorld().GetRelativeTransform(Landscape->SplineComponent->ComponentToWorld);
 
-		GUnrealEd->RedrawLevelEditingViewports();
+			AddControlPoint(Landscape, LandscapeToSpline.TransformPosition(InHitLocation));
 
-		return true;
+			GUnrealEd->RedrawLevelEditingViewports();
+
+			return true;
+		}
+
+		return false;
 	}
 
 	virtual void EndTool(FEditorViewportClient* ViewportClient) override
diff --git a/Engine/Source/Editor/UnrealEd/Private/EditorViewportClient.cpp b/Engine/Source/Editor/UnrealEd/Private/EditorViewportClient.cpp
--- a/Engine/Source/Editor/UnrealEd/Private/EditorViewportClient.cpp
+++ b/Engine/Source/Editor/UnrealEd/Private/EditorViewportClient.cpp
@@ -723,7 +723,7 @@ void FEditorViewportClient::ReceivedFocus(FViewport* InViewport)
 
 void FEditorViewportClient::LostFocus(FViewport* InViewport)
 {
-	StopTracking();	
+	StopTracking();
 }
 
 void FEditorViewportClient::Tick(float DeltaTime)
@@ -3147,6 +3147,7 @@ bool FEditorViewportClient::IsFlightCameraInputModeActive() const
 			const bool bIsNonOrbitMiddleMouse = bMiddleMouseButtonDown && !IsAltPressed();
 
 			const bool bIsMouseLooking =
+				bIsTracking &&
 				Widget->GetCurrentAxis() == EAxisList::None &&
 				( bLeftMouseButtonDown || bMiddleMouseButtonDown || bRightMouseButtonDown || bIsUsingTrackpad ) &&
 				!IsCtrlPressed() && !IsShiftPressed() && !IsAltPressed();
@@ -4003,7 +4004,7 @@ void FEditorViewportClient::DrawBoundingBox(FBox &Box, FCanvas* InCanvas, const
 			FCanvasTextItem TextItem( FVector2D( ScreenBoxMin.X + ((ScreenBoxMax.X - ScreenBoxMin.X) * 0.5f),ScreenBoxMin.Y), FText::FromString( InLabelText ), GEngine->GetMediumFont(), InColor );
 			TextItem.bCentreX = true;
 			InCanvas->DrawItem( TextItem );
-		}		
+		}
 	}
 }
 
