Foliage tool now ignores ctrl key (fix ctrl+click not painting, people are used to that key combo)
Fix for not being able to shift+click to erase with the foliage tool (whoops)

diff --git a/Engine/Source/Editor/FoliageEdit/Private/FoliageEdMode.cpp b/Engine/Source/Editor/FoliageEdit/Private/FoliageEdMode.cpp
--- a/Engine/Source/Editor/FoliageEdit/Private/FoliageEdMode.cpp
+++ b/Engine/Source/Editor/FoliageEdit/Private/FoliageEdMode.cpp
@@ -1835,8 +1835,9 @@ bool FEdModeFoliage::InputKey(FEditorViewportClient* ViewportClient, FViewport*
 	{
 		if (Key == EKeys::LeftMouseButton && Event == IE_Pressed)
 		{
-			if (!Viewport->KeyState(EKeys::MiddleMouseButton) && !Viewport->KeyState(EKeys::RightMouseButton)
-				&& !IsCtrlDown(Viewport) && !IsShiftDown(Viewport) && !IsAltDown(Viewport))
+			// Only activate tool if we're not already moving the camera and we're not trying to drag a transform widget
+			// Not using "if (!ViewportClient->IsMovingCamera())" because it's wrong in ortho viewports :D
+			if (!Viewport->KeyState(EKeys::MiddleMouseButton) && !Viewport->KeyState(EKeys::RightMouseButton) && !IsAltDown(Viewport) && ViewportClient->GetCurrentWidgetAxis() == EAxisList::None)
 			{
 				if (!bToolActive)
 				{
