Add tutorial button to Landscape and Foliage mode panels

diff --git a/Engine/Config/BaseEditorGameAgnostic.ini b/Engine/Config/BaseEditorGameAgnostic.ini
--- a/Engine/Config/BaseEditorGameAgnostic.ini
+++ b/Engine/Config/BaseEditorGameAgnostic.ini
@@ -37,3 +37,6 @@ bCaptureAudioFromMicrophone=true
 +TutorialContexts=(Context="Persona",BrowserFilter=,AttractTutorial=None,LaunchTutorial=/Engine/Tutorial/SubEditors/PersonaAnimEditorTutorial.PersonaAnimEditorTutorial_C)
 +TutorialContexts=(Context="SpriteEditor",BrowserFilter=,AttractTutorial=None,LaunchTutorial=/Engine/Tutorial/Paper2D/SpriteEditorTutorial.SpriteEditorTutorial_C)
 +TutorialContexts=(Context="FlipbookEditor",BrowserFilter=,AttractTutorial=None,LaunchTutorial=/Engine/Tutorial/Paper2D/FlipbookEditorTutorial.FlipbookEditorTutorial_C)
++TutorialContexts=(Context="LandscapeMode",BrowserFilter=,AttractTutorial=None,LaunchTutorial=/Engine/Tutorial/Landscape/Landscape_Manage_Mode.Landscape_Manage_Mode_C)
++TutorialContexts=(Context="FoliageMode",BrowserFilter=,AttractTutorial=None,LaunchTutorial=/Engine/Tutorial/Foliage/Foliage_Intro_Tutorial.Foliage_Intro_Tutorial_C)
+
diff --git a/Engine/Source/Editor/FoliageEdit/Private/SFoliageEdit.cpp b/Engine/Source/Editor/FoliageEdit/Private/SFoliageEdit.cpp
--- a/Engine/Source/Editor/FoliageEdit/Private/SFoliageEdit.cpp
+++ b/Engine/Source/Editor/FoliageEdit/Private/SFoliageEdit.cpp
@@ -7,9 +7,9 @@
 #include "FoliageEdMode.h"
 #include "Foliage/FoliageType.h"
 #include "Editor/UnrealEd/Public/AssetThumbnail.h"
-
 #include "Editor/UnrealEd/Public/DragAndDrop/AssetDragDropOp.h"
 #include "Editor/UnrealEd/Public/AssetSelection.h"
+#include "Editor/IntroTutorials/Public/IIntroTutorials.h"
 
 #include "SFoliageEditMeshDisplayItem.h"
 
@@ -101,7 +101,6 @@ void SFoliageEdit::Construct(const FArguments& InArgs)
 
 			+ SVerticalBox::Slot()
 			.AutoHeight()
-			.HAlign(HAlign_Center)
 			.Padding(StandardPadding)
 			[
 				BuildToolBar()
@@ -461,19 +460,34 @@ TSharedRef<SWidget> SFoliageEdit::BuildToolBar()
 		Toolbar.AddToolBarButton(FFoliageEditCommands::Get().SetPaintBucket);
 	}
 
+	IIntroTutorials& IntroTutorials = FModuleManager::LoadModuleChecked<IIntroTutorials>(TEXT("IntroTutorials"));
+
 	return
 		SNew(SHorizontalBox)
 
 		+ SHorizontalBox::Slot()
-		.AutoWidth()
 		.Padding(4, 0)
 		[
-			SNew(SBorder)
-			.Padding(0)
-			.BorderImage(FEditorStyle::GetBrush("NoBorder"))
-			.IsEnabled(FSlateApplication::Get().GetNormalExecutionAttribute())
+			SNew(SOverlay)
+			+ SOverlay::Slot()
+			[
+				SNew(SBorder)
+				.HAlign(HAlign_Center)
+				.Padding(0)
+				.BorderImage(FEditorStyle::GetBrush("NoBorder"))
+				.IsEnabled(FSlateApplication::Get().GetNormalExecutionAttribute())
+				[
+					Toolbar.MakeWidget()
+				]
+			]
+
+			// Tutorial link
+			+ SOverlay::Slot()
+			.HAlign(HAlign_Right)
+			.VAlign(VAlign_Bottom)
+			.Padding(4)
 			[
-				Toolbar.MakeWidget()
+				IntroTutorials.CreateTutorialsWidget(TEXT("FoliageMode"))
 			]
 		];
 }
diff --git a/Engine/Source/Editor/LandscapeEditor/Private/SLandscapeEditor.cpp b/Engine/Source/Editor/LandscapeEditor/Private/SLandscapeEditor.cpp
--- a/Engine/Source/Editor/LandscapeEditor/Private/SLandscapeEditor.cpp
+++ b/Engine/Source/Editor/LandscapeEditor/Private/SLandscapeEditor.cpp
@@ -8,6 +8,7 @@
 #include "LandscapeEditorObject.h"
 #include "IDetailsView.h"
 #include "PropertyEditorModule.h"
+#include "Editor/IntroTutorials/Public/IIntroTutorials.h"
 
 #define LOCTEXT_NAMESPACE "LandscapeEditor"
 
@@ -139,6 +140,8 @@ void SLandscapeEditor::Construct(const FArguments& InArgs, TSharedRef<FLandscape
 		DetailsPanel->SetObject(LandscapeEdMode->UISettings);
 	}
 
+	IIntroTutorials& IntroTutorials = FModuleManager::LoadModuleChecked<IIntroTutorials>(TEXT("IntroTutorials"));
+
 	ChildSlot
 	[
 		SNew(SVerticalBox)
@@ -156,13 +159,27 @@ void SLandscapeEditor::Construct(const FArguments& InArgs, TSharedRef<FLandscape
 
 			+ SVerticalBox::Slot()
 			.AutoHeight()
-			.Padding(0, 0, 0, 5)
+			.Padding(4, 0, 4, 5)
 			[
-				SNew(SBorder)
-				.BorderImage(FEditorStyle::GetBrush("ToolPanel.GroupBorder"))
-				.HAlign(HAlign_Center)
+				SNew(SOverlay)
+				+ SOverlay::Slot()
+				[
+					SNew(SBorder)
+					.BorderImage(FEditorStyle::GetBrush("ToolPanel.GroupBorder"))
+					.HAlign(HAlign_Center)
+					[
+
+						ModeSwitchButtons.MakeWidget()
+					]
+				]
+
+				// Tutorial link
+				+ SOverlay::Slot()
+				.HAlign(HAlign_Right)
+				.VAlign(VAlign_Bottom)
+				.Padding(4)
 				[
-					ModeSwitchButtons.MakeWidget()
+					IntroTutorials.CreateTutorialsWidget(TEXT("LandscapeMode"))
 				]
 			]
 			+ SVerticalBox::Slot()
