Added metatags for foliage details panel items

diff --git a/Engine/Source/Editor/FoliageEdit/Private/SFoliageEditMeshDisplayItem.cpp b/Engine/Source/Editor/FoliageEdit/Private/SFoliageEditMeshDisplayItem.cpp
--- a/Engine/Source/Editor/FoliageEdit/Private/SFoliageEditMeshDisplayItem.cpp
+++ b/Engine/Source/Editor/FoliageEdit/Private/SFoliageEditMeshDisplayItem.cpp
@@ -14,6 +14,7 @@
 
 #include "SFoliageEditMeshDisplayItem.h"
 #include "ObjectEditorUtils.h"
+#include "TutorialMetaData.h"
 
 #define LOCTEXT_NAMESPACE "FoliageEd_Mode"
 
@@ -58,6 +59,7 @@ void SFoliageEditMeshDisplayItem::Construct(const FArguments& InArgs)
 	}
 
 	SAssignNew(ThumbnailBox, SBox)
+		.AddMetaData<FTutorialMetaData>(FTutorialMetaData(TEXT("Foliage.MeshImage"), "LevelEditorToolbox"))
 		.WidthOverride(80).HeightOverride(80)
 		[
 			SAssignNew(ThumbnailWidgetBorder, SBorder)
@@ -85,6 +87,7 @@ void SFoliageEditMeshDisplayItem::Construct(const FArguments& InArgs)
 			[
 				SNew(SHorizontalBox)
 				.Visibility(this, &SFoliageEditMeshDisplayItem::IsNotReapplySettingsVisible)
+				.AddMetaData<FTutorialMetaData>(FTutorialMetaData(TEXT("Foliage.Density"), "LevelEditorToolbox"))
 
 				+ SHorizontalBox::Slot()
 				.FillWidth(MAIN_TITLE + SPINBOX_PREFIX)
@@ -163,6 +166,7 @@ void SFoliageEditMeshDisplayItem::Construct(const FArguments& InArgs)
 			.Padding(StandardPadding)
 			[
 				SNew(SHorizontalBox)
+				.AddMetaData<FTutorialMetaData>(FTutorialMetaData(TEXT("Foliage.Radius"), "LevelEditorToolbox"))
 
 				+ SHorizontalBox::Slot()
 				.AutoWidth()
@@ -203,6 +207,7 @@ void SFoliageEditMeshDisplayItem::Construct(const FArguments& InArgs)
 
 	TSharedRef<SHorizontalBox> AlignToNormalBox =
 		SNew(SHorizontalBox)
+		.AddMetaData<FTutorialMetaData>(FTutorialMetaData(TEXT("Foliage.AlignToNormal"), "LevelEditorToolbox"))			
 
 		+ SHorizontalBox::Slot()
 		.AutoWidth()
@@ -229,6 +234,7 @@ void SFoliageEditMeshDisplayItem::Construct(const FArguments& InArgs)
 	TSharedRef<SHorizontalBox> MaxAngleBox =
 		SNew(SHorizontalBox)
 		.Visibility(this, &SFoliageEditMeshDisplayItem::IsAlignToNormalVisible)
+		.AddMetaData<FTutorialMetaData>(FTutorialMetaData(TEXT("Foliage.MaxAngle"), "LevelEditorToolbox"))
 
 		// Dummy Checkbox
 		+ SHorizontalBox::Slot()
@@ -272,6 +278,7 @@ void SFoliageEditMeshDisplayItem::Construct(const FArguments& InArgs)
 
 	TSharedRef<SHorizontalBox> RandomYawBox =
 		SNew(SHorizontalBox)
+		.AddMetaData<FTutorialMetaData>(FTutorialMetaData(TEXT("Foliage.RandomYaw"), "LevelEditorToolbox"))
 
 		+ SHorizontalBox::Slot()
 		.AutoWidth()
@@ -297,6 +304,7 @@ void SFoliageEditMeshDisplayItem::Construct(const FArguments& InArgs)
 
 	TSharedRef<SHorizontalBox> UniformScaleBox =
 		SNew(SHorizontalBox)
+		.AddMetaData<FTutorialMetaData>(FTutorialMetaData(TEXT("Foliage.UniformScale"), "LevelEditorToolbox"))
 
 		// Dummy Checkbox
 		+ SHorizontalBox::Slot()
@@ -338,6 +346,7 @@ void SFoliageEditMeshDisplayItem::Construct(const FArguments& InArgs)
 	TSharedRef<SHorizontalBox> ScaleUniformBox =
 		SNew(SHorizontalBox)
 		.Visibility(this, &SFoliageEditMeshDisplayItem::IsUniformScalingVisible)
+		.AddMetaData<FTutorialMetaData>(FTutorialMetaData(TEXT("Foliage.ScaleValues"), "LevelEditorToolbox"))
 
 		+ SHorizontalBox::Slot()
 		.AutoWidth()
@@ -664,6 +673,7 @@ void SFoliageEditMeshDisplayItem::Construct(const FArguments& InArgs)
 
 	TSharedRef<SHorizontalBox> ZOffsetBox =
 		SNew(SHorizontalBox)
+		.AddMetaData<FTutorialMetaData>(FTutorialMetaData(TEXT("Foliage.ZOffset"), "LevelEditorToolbox"))
 
 		+ SHorizontalBox::Slot()
 		.AutoWidth()
@@ -746,6 +756,7 @@ void SFoliageEditMeshDisplayItem::Construct(const FArguments& InArgs)
 
 	TSharedRef<SHorizontalBox> RandomPitchBox =
 		SNew(SHorizontalBox)
+		.AddMetaData<FTutorialMetaData>(FTutorialMetaData(TEXT("Foliage.RandomPitch"), "LevelEditorToolbox"))
 
 		+ SHorizontalBox::Slot()
 		.AutoWidth()
@@ -789,6 +800,7 @@ void SFoliageEditMeshDisplayItem::Construct(const FArguments& InArgs)
 		];
 
 	TSharedRef<SHorizontalBox> GroundSlopeBox = SNew(SHorizontalBox)
+		.AddMetaData<FTutorialMetaData>(FTutorialMetaData(TEXT("Foliage.GroundSlope"), "LevelEditorToolbox"))
 
 		+ SHorizontalBox::Slot()
 		.AutoWidth()
@@ -833,6 +845,7 @@ void SFoliageEditMeshDisplayItem::Construct(const FArguments& InArgs)
 
 	TSharedRef<SHorizontalBox> HeightBox =
 		SNew(SHorizontalBox)
+		.AddMetaData<FTutorialMetaData>(FTutorialMetaData(TEXT("Foliage.Height"), "LevelEditorToolbox"))
 
 		+ SHorizontalBox::Slot()
 		.AutoWidth()
@@ -910,6 +923,7 @@ void SFoliageEditMeshDisplayItem::Construct(const FArguments& InArgs)
 
 	TSharedRef<SHorizontalBox> LandscapeLayerBox =
 		SNew(SHorizontalBox)
+		.AddMetaData<FTutorialMetaData>(FTutorialMetaData(TEXT("Foliage.LandscapeLayer"), "LevelEditorToolbox"))
 
 		+ SHorizontalBox::Slot()
 		.AutoWidth()
@@ -951,6 +965,7 @@ void SFoliageEditMeshDisplayItem::Construct(const FArguments& InArgs)
 
 	TSharedRef<SHorizontalBox> CollisionWithWorldBox =
 		SNew(SHorizontalBox)
+		.AddMetaData<FTutorialMetaData>(FTutorialMetaData(TEXT("Foliage.CollisionWorld"), "LevelEditorToolbox"))
 
 		+ SHorizontalBox::Slot()
 		.AutoWidth()
@@ -1078,6 +1093,7 @@ void SFoliageEditMeshDisplayItem::Construct(const FArguments& InArgs)
 
 	TSharedRef<SHorizontalBox> VertexColorMaskBox =
 		SNew(SHorizontalBox)
+		.AddMetaData<FTutorialMetaData>(FTutorialMetaData(TEXT("Foliage.VertexColorMask"), "LevelEditorToolbox"))
 
 		+ SHorizontalBox::Slot()
 		.AutoWidth()
@@ -1452,6 +1468,7 @@ void SFoliageEditMeshDisplayItem::Construct(const FArguments& InArgs)
 						.ButtonStyle(FEditorStyle::Get(), "NoBorder")
 						.OnClicked(this, &SFoliageEditMeshDisplayItem::OnReplace)
 						.ToolTipText(NSLOCTEXT("FoliageEdMode", "Replace_Tooltip", "Replace all instances with the Static Mesh currently selected in the Content Browser."))
+						.AddMetaData<FTutorialMetaData>(FTutorialMetaData(TEXT("Foliage.ReplaceInstances"), "LevelEditorToolbox"))
 						[
 							SNew(SImage)
 							.Image(FEditorStyle::GetBrush(TEXT("ContentReference.UseSelectionFromContentBrowser")))
@@ -1466,6 +1483,7 @@ void SFoliageEditMeshDisplayItem::Construct(const FArguments& InArgs)
 						.ButtonStyle(FEditorStyle::Get(), "NoBorder")
 						.OnClicked(this, &SFoliageEditMeshDisplayItem::OnSync)
 						.ToolTipText(NSLOCTEXT("FoliageEdMode", "FindInContentBrowser_Tooltip", "Find this Static Mesh in the Content Browser."))
+						.AddMetaData<FTutorialMetaData>(FTutorialMetaData(TEXT("Foliage.FindInBrowser"), "LevelEditorToolbox"))
 						[
 							SNew(SImage)
 							.Image(FEditorStyle::GetBrush(TEXT("ContentReference.FindInContentBrowser")))
@@ -1480,6 +1498,7 @@ void SFoliageEditMeshDisplayItem::Construct(const FArguments& InArgs)
 						.ButtonStyle(FEditorStyle::Get(), "NoBorder")
 						.OnClicked(this, &SFoliageEditMeshDisplayItem::OnRemove)
 						.ToolTipText(NSLOCTEXT("FoliageEdMode", "Remove_Tooltip", "Delete all foliage instances of this Static Mesh."))
+						.AddMetaData<FTutorialMetaData>(FTutorialMetaData(TEXT("Foliage.DeleteInstances"), "LevelEditorToolbox"))
 						[
 							SNew(SImage)
 							.Image(FEditorStyle::GetBrush(TEXT("ContentReference.Clear")))
@@ -1510,6 +1529,7 @@ void SFoliageEditMeshDisplayItem::Construct(const FArguments& InArgs)
 						.ButtonStyle(FEditorStyle::Get(), "NoBorder")
 						.OnClicked(this, &SFoliageEditMeshDisplayItem::OnOpenSettings)
 						.ToolTipText(NSLOCTEXT("FoliageEdMode", "OpenSettings_Tooltip", "Use the InstancedFoliageSettings currently selected in the Content Browser."))
+						.AddMetaData<FTutorialMetaData>(FTutorialMetaData(TEXT("Foliage.OpenSettings"), "LevelEditorToolbox"))
 						[
 							SNew(SImage)
 							.Image(FEditorStyle::GetBrush(TEXT("FoliageEditMode.OpenSettings")))
@@ -1523,6 +1543,7 @@ void SFoliageEditMeshDisplayItem::Construct(const FArguments& InArgs)
 						.ButtonStyle(FEditorStyle::Get(), "NoBorder")
 						.OnClicked(this, &SFoliageEditMeshDisplayItem::OnSaveRemoveSettings)
 						.ToolTipText(this, &SFoliageEditMeshDisplayItem::GetSaveRemoveSettingsTooltip)
+						.AddMetaData<FTutorialMetaData>(FTutorialMetaData(TEXT("Foliage.SaveSettings"), "LevelEditorToolbox"))
 						[
 							SNew(SImage)
 							.Image(this, &SFoliageEditMeshDisplayItem::GetSaveSettingsBrush)
