Added missing Foliage and Landscape tooltips

TTP#346285, 345541

diff --git a/Engine/Source/Editor/FoliageEdit/Private/SFoliageEdit.cpp b/Engine/Source/Editor/FoliageEdit/Private/SFoliageEdit.cpp
--- a/Engine/Source/Editor/FoliageEdit/Private/SFoliageEdit.cpp
+++ b/Engine/Source/Editor/FoliageEdit/Private/SFoliageEdit.cpp
@@ -110,6 +110,7 @@ void SFoliageEdit::Construct(const FArguments& InArgs)
 			.AutoHeight()
 			[
 				SNew(SHorizontalBox)
+				.ToolTipText(LOCTEXT("BrushSize_Tooltip", "The size of the foliage brush"))
 
 				+ SHorizontalBox::Slot()
 				.FillWidth(1.0f)
@@ -138,6 +139,7 @@ void SFoliageEdit::Construct(const FArguments& InArgs)
 			.AutoHeight()
 			[
 				SNew(SHorizontalBox)
+				.ToolTipText(LOCTEXT("PaintDensity_Tooltip", "The density of foliage to paint. This is a multiplier for the individual foliage type's density specifier."))
 
 				+ SHorizontalBox::Slot()
 				.FillWidth(1.0f)
@@ -165,6 +167,7 @@ void SFoliageEdit::Construct(const FArguments& InArgs)
 			.AutoHeight()
 			[
 				SNew(SHorizontalBox)
+				.ToolTipText(LOCTEXT("EraseDensity_Tooltip", "The density of foliage to leave behind when erasing with the Shift key held. 0 will remove all foliage."))
 
 				+ SHorizontalBox::Slot()
 				.FillWidth(1.0f)
@@ -211,6 +214,7 @@ void SFoliageEdit::Construct(const FArguments& InArgs)
 					.Visibility(this, &SFoliageEdit::GetVisibility_Filters)
 					.OnCheckStateChanged(this, &SFoliageEdit::OnCheckStateChanged_Landscape)
 					.IsChecked(this, &SFoliageEdit::GetCheckState_Landscape)
+					.ToolTipText(LOCTEXT("FilterLandscape_Tooltip", "Place foliage on Landscape"))
 					[
 						SNew(STextBlock)
 						.Text(LOCTEXT("Landscape", "Landscape"))
@@ -224,6 +228,7 @@ void SFoliageEdit::Construct(const FArguments& InArgs)
 					.Visibility(this, &SFoliageEdit::GetVisibility_Filters)
 					.OnCheckStateChanged(this, &SFoliageEdit::OnCheckStateChanged_StaticMesh)
 					.IsChecked(this, &SFoliageEdit::GetCheckState_StaticMesh)
+					.ToolTipText(LOCTEXT("FilterStaticMesh_Tooltip", "Place foliage on StaticMeshes"))
 					[
 						SNew(STextBlock)
 						.Text(LOCTEXT("StaticMeshes", "Static Meshes"))
@@ -237,6 +242,7 @@ void SFoliageEdit::Construct(const FArguments& InArgs)
 					.Visibility(this, &SFoliageEdit::GetVisibility_Filters)
 					.OnCheckStateChanged(this, &SFoliageEdit::OnCheckStateChanged_BSP)
 					.IsChecked(this, &SFoliageEdit::GetCheckState_BSP)
+					.ToolTipText(LOCTEXT("FilterBSP_Tooltip", "Place foliage on BSP"))
 					[
 						SNew(STextBlock)
 						.Text(LOCTEXT("BSP", "BSP"))
@@ -250,6 +256,7 @@ void SFoliageEdit::Construct(const FArguments& InArgs)
 					.Visibility(this, &SFoliageEdit::GetVisibility_Filters)
 					.OnCheckStateChanged(this, &SFoliageEdit::OnCheckStateChanged_Translucent)
 					.IsChecked(this, &SFoliageEdit::GetCheckState_Translucent)
+					.ToolTipText(LOCTEXT("FilterTranslucent_Tooltip", "Place foliage on translucent surfaces"))
 					[
 						SNew(STextBlock)
 						.Text(LOCTEXT("Translucent", "Translucent"))
diff --git a/Engine/Source/Editor/FoliageEdit/Private/SFoliageEditMeshDisplayItem.cpp b/Engine/Source/Editor/FoliageEdit/Private/SFoliageEditMeshDisplayItem.cpp
--- a/Engine/Source/Editor/FoliageEdit/Private/SFoliageEditMeshDisplayItem.cpp
+++ b/Engine/Source/Editor/FoliageEdit/Private/SFoliageEditMeshDisplayItem.cpp
@@ -88,6 +88,7 @@ void SFoliageEditMeshDisplayItem::Construct(const FArguments& InArgs)
 				SNew(SHorizontalBox)
 				.Visibility(this, &SFoliageEditMeshDisplayItem::IsNotReapplySettingsVisible)
 				.AddMetaData<FTutorialMetaData>(FTutorialMetaData(TEXT("Foliage.Density"), "LevelEditorToolbox"))
+				.ToolTipText(LOCTEXT("Density_Tooltip", "Foliage instances will be placed at this density, specified in instances per 1000x1000 unit area."))
 
 				+ SHorizontalBox::Slot()
 				.FillWidth(MAIN_TITLE + SPINBOX_PREFIX)
@@ -131,6 +132,7 @@ void SFoliageEditMeshDisplayItem::Construct(const FArguments& InArgs)
 			[
 				SNew(SHorizontalBox)
 				.Visibility(this, &SFoliageEditMeshDisplayItem::IsReapplySettingsVisible)
+				.ToolTipText(LOCTEXT("DensityAdjust_Tooltip", "If checked, the density of foliage instances already placed will be adjusted. <1 will remove instances to reduce the density, >1 will place extra instances."))
 
 				+ SHorizontalBox::Slot()
 				.AutoWidth()
@@ -167,6 +169,7 @@ void SFoliageEditMeshDisplayItem::Construct(const FArguments& InArgs)
 			[
 				SNew(SHorizontalBox)
 				.AddMetaData<FTutorialMetaData>(FTutorialMetaData(TEXT("Foliage.Radius"), "LevelEditorToolbox"))
+				.ToolTipText(LOCTEXT("Radius_Tooltip", "The minimum distance between foliage instances."))
 
 				+ SHorizontalBox::Slot()
 				.AutoWidth()
@@ -175,6 +178,7 @@ void SFoliageEditMeshDisplayItem::Construct(const FArguments& InArgs)
 					.Visibility(this, &SFoliageEditMeshDisplayItem::IsReapplySettingsVisible)
 					.OnCheckStateChanged(this, &SFoliageEditMeshDisplayItem::OnRadiusReapply)
 					.IsChecked(this, &SFoliageEditMeshDisplayItem::IsRadiusReapplyChecked)
+					.ToolTipText(LOCTEXT("ReapplyRadiusCheck_Tooltip", "If checked, foliage instances not meeting the new Radius constraint will be removed"))
 				]
 
 				+ SHorizontalBox::Slot()
@@ -208,6 +212,7 @@ void SFoliageEditMeshDisplayItem::Construct(const FArguments& InArgs)
 	TSharedRef<SHorizontalBox> AlignToNormalBox =
 		SNew(SHorizontalBox)
 		.AddMetaData<FTutorialMetaData>(FTutorialMetaData(TEXT("Foliage.AlignToNormal"), "LevelEditorToolbox"))			
+		.ToolTipText(LOCTEXT("AlignToNormal_Tooltip", "Whether foliage instances should have their angle adjusted away from vertical to match the normal of the surface they're painted on"))
 
 		+ SHorizontalBox::Slot()
 		.AutoWidth()
@@ -216,6 +221,7 @@ void SFoliageEditMeshDisplayItem::Construct(const FArguments& InArgs)
 			.Visibility(this, &SFoliageEditMeshDisplayItem::IsReapplySettingsVisible)
 			.OnCheckStateChanged(this, &SFoliageEditMeshDisplayItem::OnAlignToNormalReapply)
 			.IsChecked(this, &SFoliageEditMeshDisplayItem::IsAlignToNormalReapplyChecked)
+			.ToolTipText(LOCTEXT("ReapplyNormalCheck_Tooltip", "If checked, foliage instances will have their normal alignment adjusted by the Reapply tool"))
 		]
 
 		+ SHorizontalBox::Slot()
@@ -235,6 +241,7 @@ void SFoliageEditMeshDisplayItem::Construct(const FArguments& InArgs)
 		SNew(SHorizontalBox)
 		.Visibility(this, &SFoliageEditMeshDisplayItem::IsAlignToNormalVisible)
 		.AddMetaData<FTutorialMetaData>(FTutorialMetaData(TEXT("Foliage.MaxAngle"), "LevelEditorToolbox"))
+		.ToolTipText(LOCTEXT("MaxAngle_Tooltip", "The maximum angle in degrees that foliage instances will be adjusted away from the vertical"))
 
 		// Dummy Checkbox
 		+ SHorizontalBox::Slot()
@@ -279,6 +286,7 @@ void SFoliageEditMeshDisplayItem::Construct(const FArguments& InArgs)
 	TSharedRef<SHorizontalBox> RandomYawBox =
 		SNew(SHorizontalBox)
 		.AddMetaData<FTutorialMetaData>(FTutorialMetaData(TEXT("Foliage.RandomYaw"), "LevelEditorToolbox"))
+		.ToolTipText(LOCTEXT("RandomYaw_Tooltip", "If selected, foliage instances will have a random yaw rotation around their vertical axis applied"))
 
 		+ SHorizontalBox::Slot()
 		.AutoWidth()
@@ -287,6 +295,7 @@ void SFoliageEditMeshDisplayItem::Construct(const FArguments& InArgs)
 			.Visibility(this, &SFoliageEditMeshDisplayItem::IsReapplySettingsVisible)
 			.OnCheckStateChanged(this, &SFoliageEditMeshDisplayItem::OnRandomYawReapply)
 			.IsChecked(this, &SFoliageEditMeshDisplayItem::IsRandomYawReapplyChecked)
+			.ToolTipText(LOCTEXT("ReapplyRandomYawCheck_Tooltip", "If checked, foliage instances will have their yaw adjusted by the Reapply tool"))
 		]
 
 		+ SHorizontalBox::Slot()
@@ -305,6 +314,7 @@ void SFoliageEditMeshDisplayItem::Construct(const FArguments& InArgs)
 	TSharedRef<SHorizontalBox> UniformScaleBox =
 		SNew(SHorizontalBox)
 		.AddMetaData<FTutorialMetaData>(FTutorialMetaData(TEXT("Foliage.UniformScale"), "LevelEditorToolbox"))
+		.ToolTipText(LOCTEXT("UniformScale_Tooltip", "If selected, foliage instances will have unfiorm  X,Y and Z scales"))
 
 		// Dummy Checkbox
 		+ SHorizontalBox::Slot()
@@ -347,6 +357,7 @@ void SFoliageEditMeshDisplayItem::Construct(const FArguments& InArgs)
 		SNew(SHorizontalBox)
 		.Visibility(this, &SFoliageEditMeshDisplayItem::IsUniformScalingVisible)
 		.AddMetaData<FTutorialMetaData>(FTutorialMetaData(TEXT("Foliage.ScaleValues"), "LevelEditorToolbox"))
+		.ToolTipText(LOCTEXT("ScaleUniformValues_Tooltip", "Specifies the range of scale, from minimum to maximum, to apply to a foliage instance's Scale property, applied uniformly to X, Y and Z."))
 
 		+ SHorizontalBox::Slot()
 		.AutoWidth()
@@ -355,6 +366,7 @@ void SFoliageEditMeshDisplayItem::Construct(const FArguments& InArgs)
 			.Visibility(this, &SFoliageEditMeshDisplayItem::IsReapplySettingsVisible)
 			.OnCheckStateChanged(this, &SFoliageEditMeshDisplayItem::OnScaleUniformReapply)
 			.IsChecked(this, &SFoliageEditMeshDisplayItem::IsScaleUniformReapplyChecked)
+			.ToolTipText(LOCTEXT("ReapplyScaleUnformCheck_Tooltip", "If checked, foliage instances will have their scale uniformly adjusted by the Reapply tool"))
 		]
 
 		+ SHorizontalBox::Slot()
@@ -428,6 +440,7 @@ void SFoliageEditMeshDisplayItem::Construct(const FArguments& InArgs)
 	TSharedRef<SHorizontalBox> ScaleXBox =
 		SNew(SHorizontalBox)
 		.Visibility(this, &SFoliageEditMeshDisplayItem::IsNonUniformScalingVisible)
+		.ToolTipText(LOCTEXT("ScaleXValues_Tooltip", "Specifies the range of scale, from minimum to maximum, to apply to a foliage instance's X Scale property"))
 
 		+ SHorizontalBox::Slot()
 		.AutoWidth()
@@ -436,6 +449,7 @@ void SFoliageEditMeshDisplayItem::Construct(const FArguments& InArgs)
 			.Visibility(this, &SFoliageEditMeshDisplayItem::IsReapplySettingsVisible)
 			.OnCheckStateChanged(this, &SFoliageEditMeshDisplayItem::OnScaleXReapply)
 			.IsChecked(this, &SFoliageEditMeshDisplayItem::IsScaleXReapplyChecked)
+			.ToolTipText(LOCTEXT("ReapplyScaleXCheck_Tooltip", "If checked, foliage instances will have their X scale adjusted by the Reapply tool"))
 		]
 
 		+ SHorizontalBox::Slot()
@@ -504,12 +518,14 @@ void SFoliageEditMeshDisplayItem::Construct(const FArguments& InArgs)
 				SNew(SCheckBox)
 				.OnCheckStateChanged(this, &SFoliageEditMeshDisplayItem::OnScaleXLocked)
 				.IsChecked(this, &SFoliageEditMeshDisplayItem::IsScaleXLockedChecked)
+				.ToolTipText(LOCTEXT("ScaleXLocked_Tooltip", "Locks the X axis scale. All axes with the locked checkbox set will have the same scale applied."))
 			]
 		];
 
 	TSharedRef<SHorizontalBox> ScaleYBox =
 		SNew(SHorizontalBox)
 		.Visibility(this, &SFoliageEditMeshDisplayItem::IsNonUniformScalingVisible)
+		.ToolTipText(LOCTEXT("ScaleYValues_Tooltip", "Specifies the range of scale, from minimum to maximum, to apply to a foliage instance's Y Scale property"))
 
 		+ SHorizontalBox::Slot()
 		.AutoWidth()
@@ -518,6 +534,7 @@ void SFoliageEditMeshDisplayItem::Construct(const FArguments& InArgs)
 			.Visibility(this, &SFoliageEditMeshDisplayItem::IsReapplySettingsVisible)
 			.OnCheckStateChanged(this, &SFoliageEditMeshDisplayItem::OnScaleYReapply)
 			.IsChecked(this, &SFoliageEditMeshDisplayItem::IsScaleYReapplyChecked)
+			.ToolTipText(LOCTEXT("ReapplyScaleYCheck_Tooltip", "If checked, foliage instances will have their Y scale adjusted by the Reapply tool"))
 		]
 
 		+ SHorizontalBox::Slot()
@@ -586,12 +603,14 @@ void SFoliageEditMeshDisplayItem::Construct(const FArguments& InArgs)
 				SNew(SCheckBox)
 				.OnCheckStateChanged(this, &SFoliageEditMeshDisplayItem::OnScaleYLocked)
 				.IsChecked(this, &SFoliageEditMeshDisplayItem::IsScaleYLockedChecked)
+				.ToolTipText(LOCTEXT("ScaleYLocked_Tooltip", "Locks the Y axis scale. All axes with the locked checkbox set will have the same scale applied."))
 			]
 		];
 
 	TSharedRef<SHorizontalBox> ScaleZBox =
 		SNew(SHorizontalBox)
 		.Visibility(this, &SFoliageEditMeshDisplayItem::IsNonUniformScalingVisible)
+		.ToolTipText(LOCTEXT("ScaleZValues_Tooltip", "Specifies the range of scale, from minimum to maximum, to apply to a foliage instance's Z Scale property"))
 
 		+ SHorizontalBox::Slot()
 		.AutoWidth()
@@ -600,6 +619,7 @@ void SFoliageEditMeshDisplayItem::Construct(const FArguments& InArgs)
 			.Visibility(this, &SFoliageEditMeshDisplayItem::IsReapplySettingsVisible)
 			.OnCheckStateChanged(this, &SFoliageEditMeshDisplayItem::OnScaleZReapply)
 			.IsChecked(this, &SFoliageEditMeshDisplayItem::IsScaleZReapplyChecked)
+			.ToolTipText(LOCTEXT("ReapplyScaleZCheck_Tooltip", "If checked, foliage instances will have their Z scale adjusted by the Reapply tool"))
 		]
 
 		+ SHorizontalBox::Slot()
@@ -668,12 +688,14 @@ void SFoliageEditMeshDisplayItem::Construct(const FArguments& InArgs)
 				SNew(SCheckBox)
 				.OnCheckStateChanged(this, &SFoliageEditMeshDisplayItem::OnScaleZLocked)
 				.IsChecked(this, &SFoliageEditMeshDisplayItem::IsScaleZLockedChecked)
+				.ToolTipText(LOCTEXT("ScaleZLocked_Tooltip", "Locks the Z axis scale. All axes with the locked checkbox set will have the same scale applied."))
 			]
 		];
 
 	TSharedRef<SHorizontalBox> ZOffsetBox =
 		SNew(SHorizontalBox)
 		.AddMetaData<FTutorialMetaData>(FTutorialMetaData(TEXT("Foliage.ZOffset"), "LevelEditorToolbox"))
+		.ToolTipText(LOCTEXT("ZOffset_Tooltip", "Specifies a range from minimum to maximum of the offset to apply to a foliage instance's Z location"))
 
 		+ SHorizontalBox::Slot()
 		.AutoWidth()
@@ -682,6 +704,7 @@ void SFoliageEditMeshDisplayItem::Construct(const FArguments& InArgs)
 			.Visibility(this, &SFoliageEditMeshDisplayItem::IsReapplySettingsVisible)
 			.OnCheckStateChanged(this, &SFoliageEditMeshDisplayItem::OnZOffsetReapply)
 			.IsChecked(this, &SFoliageEditMeshDisplayItem::IsZOffsetReapplyChecked)
+			.ToolTipText(LOCTEXT("ReapplyZOffsetCheck_Tooltip", "If checked, foliage instances will have their Z offset adjusted by the Reapply tool"))
 		]
 
 		+ SHorizontalBox::Slot()
@@ -757,6 +780,7 @@ void SFoliageEditMeshDisplayItem::Construct(const FArguments& InArgs)
 	TSharedRef<SHorizontalBox> RandomPitchBox =
 		SNew(SHorizontalBox)
 		.AddMetaData<FTutorialMetaData>(FTutorialMetaData(TEXT("Foliage.RandomPitch"), "LevelEditorToolbox"))
+		.ToolTipText(LOCTEXT("RandomPitch_Tooltip", "A random pitch adjustment can be applied to each instance, up to the specified angle in degrees, from the original vertical."))
 
 		+ SHorizontalBox::Slot()
 		.AutoWidth()
@@ -765,6 +789,7 @@ void SFoliageEditMeshDisplayItem::Construct(const FArguments& InArgs)
 			.Visibility(this, &SFoliageEditMeshDisplayItem::IsReapplySettingsVisible)
 			.OnCheckStateChanged(this, &SFoliageEditMeshDisplayItem::OnRandomPitchReapply)
 			.IsChecked(this, &SFoliageEditMeshDisplayItem::IsRandomPitchReapplyChecked)
+			.ToolTipText(LOCTEXT("ReapplyRandomPitch_Tooltip", "If checked, foliage instances will have their pitch adjusted by the Reapply tool"))
 		]
 
 		+ SHorizontalBox::Slot()
@@ -801,6 +826,7 @@ void SFoliageEditMeshDisplayItem::Construct(const FArguments& InArgs)
 
 	TSharedRef<SHorizontalBox> GroundSlopeBox = SNew(SHorizontalBox)
 		.AddMetaData<FTutorialMetaData>(FTutorialMetaData(TEXT("Foliage.GroundSlope"), "LevelEditorToolbox"))
+		.ToolTipText(LOCTEXT("GroundSlope_Tooltip", "If non-zero, foliage instances will only be placed on surfaces sloping less than the angle specified from the horizontal. Negative values reverse the test, placing instances only on surfaces sloping greater than the specified angle."))
 
 		+ SHorizontalBox::Slot()
 		.AutoWidth()
@@ -809,6 +835,7 @@ void SFoliageEditMeshDisplayItem::Construct(const FArguments& InArgs)
 			.Visibility(this, &SFoliageEditMeshDisplayItem::IsReapplySettingsVisible)
 			.OnCheckStateChanged(this, &SFoliageEditMeshDisplayItem::OnGroundSlopeReapply)
 			.IsChecked(this, &SFoliageEditMeshDisplayItem::IsGroundSlopeReapplyChecked)
+			.ToolTipText(LOCTEXT("ReapplyGroundSlopeCheck_Tooltip", "If checked, foliage instances not meeting the ground slope condition will be removed by the Reapply tool"))
 		]
 
 		+ SHorizontalBox::Slot()
@@ -846,6 +873,7 @@ void SFoliageEditMeshDisplayItem::Construct(const FArguments& InArgs)
 	TSharedRef<SHorizontalBox> HeightBox =
 		SNew(SHorizontalBox)
 		.AddMetaData<FTutorialMetaData>(FTutorialMetaData(TEXT("Foliage.Height"), "LevelEditorToolbox"))
+		.ToolTipText(LOCTEXT("Height_Tooltip", "The valid altitude range where foliage instances will be placed, specified using minimum and maximum world coordinate Z values"))
 
 		+ SHorizontalBox::Slot()
 		.AutoWidth()
@@ -854,6 +882,7 @@ void SFoliageEditMeshDisplayItem::Construct(const FArguments& InArgs)
 			.Visibility(this, &SFoliageEditMeshDisplayItem::IsReapplySettingsVisible)
 			.OnCheckStateChanged(this, &SFoliageEditMeshDisplayItem::OnHeightReapply)
 			.IsChecked(this, &SFoliageEditMeshDisplayItem::IsHeightReapplyChecked)
+			.ToolTipText(LOCTEXT("ReapplyHeightCheck_Tooltip", "If checked, foliage instances not meeting the valid Z height condition will be removed by the Reapply tool"))
 		]
 		+ SHorizontalBox::Slot()
 		.FillWidth(MAIN_TITLE)
@@ -924,6 +953,7 @@ void SFoliageEditMeshDisplayItem::Construct(const FArguments& InArgs)
 	TSharedRef<SHorizontalBox> LandscapeLayerBox =
 		SNew(SHorizontalBox)
 		.AddMetaData<FTutorialMetaData>(FTutorialMetaData(TEXT("Foliage.LandscapeLayer"), "LevelEditorToolbox"))
+		.ToolTipText(LOCTEXT("LandscapeLayer_Tooltip", "If a layer name is specified, painting on landscape will limit the foliage to areas of landscape with the specified layer painted."))
 
 		+ SHorizontalBox::Slot()
 		.AutoWidth()
@@ -932,6 +962,7 @@ void SFoliageEditMeshDisplayItem::Construct(const FArguments& InArgs)
 			.Visibility(this, &SFoliageEditMeshDisplayItem::IsReapplySettingsVisible)
 			.OnCheckStateChanged(this, &SFoliageEditMeshDisplayItem::OnLandscapeLayerReapply)
 			.IsChecked(this, &SFoliageEditMeshDisplayItem::IsLandscapeLayerReapplyChecked)
+			.ToolTipText(LOCTEXT("ReapplyLandscapeLayerCheck_Tooltip", "If checked, foliage instances painted on areas that do not have the appopriate landscape layer painted will be removed by the Reapply tool"))
 		]
 
 		+ SHorizontalBox::Slot()
@@ -966,6 +997,7 @@ void SFoliageEditMeshDisplayItem::Construct(const FArguments& InArgs)
 	TSharedRef<SHorizontalBox> CollisionWithWorldBox =
 		SNew(SHorizontalBox)
 		.AddMetaData<FTutorialMetaData>(FTutorialMetaData(TEXT("Foliage.CollisionWorld"), "LevelEditorToolbox"))
+		.ToolTipText(LOCTEXT("CollisionWorld_Tooltip", "If checked, an overlap test with existing world geometry is performed before each instance is placed."))
 
 		+ SHorizontalBox::Slot()
 		.AutoWidth()
@@ -974,6 +1006,7 @@ void SFoliageEditMeshDisplayItem::Construct(const FArguments& InArgs)
 			.Visibility(this, &SFoliageEditMeshDisplayItem::IsReapplySettingsVisible)
 			.OnCheckStateChanged(this, &SFoliageEditMeshDisplayItem::OnCollisionWithWorldReapply)
 			.IsChecked(this, &SFoliageEditMeshDisplayItem::IsCollisionWithWorldReapplyChecked)
+			.ToolTipText(LOCTEXT("ReapplyCollisionWithWorldCheck_Tooltip", "If checked, foliage instances will have an overlap test with the world reapplied, and overlapping instances will be removed by the Reapply tool"))
 		]
 
 		+ SHorizontalBox::Slot()
@@ -992,6 +1025,7 @@ void SFoliageEditMeshDisplayItem::Construct(const FArguments& InArgs)
 	TSharedRef<SHorizontalBox> CollisionScaleX =
 		SNew(SHorizontalBox)
 		.Visibility(this, &SFoliageEditMeshDisplayItem::IsCollisionWithWorldVisible)
+		.ToolTipText(LOCTEXT("CollisionWorldScale_Tooltip", "The foliage instance's collision bounding box will be scaled by the specified amount before performing the overlap check."))
 
 		// Dummy Checkbox
 		+ SHorizontalBox::Slot()
@@ -1094,6 +1128,7 @@ void SFoliageEditMeshDisplayItem::Construct(const FArguments& InArgs)
 	TSharedRef<SHorizontalBox> VertexColorMaskBox =
 		SNew(SHorizontalBox)
 		.AddMetaData<FTutorialMetaData>(FTutorialMetaData(TEXT("Foliage.VertexColorMask"), "LevelEditorToolbox"))
+		.ToolTipText(LOCTEXT("VertexColorMask_Tooltip", "When painting on static meshes, foliage instance placement can be limited to areas where the static mesh has values in the selected vertex color channel(s). This allows a static mesh to mask out certain areas to prevent foliage from being placed there."))
 
 		+ SHorizontalBox::Slot()
 		.AutoWidth()
@@ -1102,6 +1137,7 @@ void SFoliageEditMeshDisplayItem::Construct(const FArguments& InArgs)
 			.Visibility(this, &SFoliageEditMeshDisplayItem::IsReapplySettingsVisible)
 			.OnCheckStateChanged(this, &SFoliageEditMeshDisplayItem::OnVertexColorMaskReapply)
 			.IsChecked(this, &SFoliageEditMeshDisplayItem::IsVertexColorMaskReapplyChecked)
+			.ToolTipText(LOCTEXT("ReapplyVertexColorMaskCheck_Tooltip", "If checked, foliage instances no longer matching the vertex color constraint will be removed by the Reapply tool"))
 		]
 
 		+ SHorizontalBox::Slot()
@@ -1125,6 +1161,7 @@ void SFoliageEditMeshDisplayItem::Construct(const FArguments& InArgs)
 				SNew(SCheckBox)
 				.OnCheckStateChanged(this, &SFoliageEditMeshDisplayItem::OnVertexColorMask, FOLIAGEVERTEXCOLORMASK_Disabled)
 				.IsChecked(this, &SFoliageEditMeshDisplayItem::IsVertexColorMaskChecked, FOLIAGEVERTEXCOLORMASK_Disabled)
+				.ToolTipText(LOCTEXT("VertexColorDisabled_Tooltip", "Do not use a vertex color mask"))
 				[
 					SNew(STextBlock)
 					.Text(LOCTEXT("VertexColorDisabled", "Disabled"))
@@ -1139,6 +1176,7 @@ void SFoliageEditMeshDisplayItem::Construct(const FArguments& InArgs)
 				SNew(SCheckBox)
 				.OnCheckStateChanged(this, &SFoliageEditMeshDisplayItem::OnVertexColorMask, FOLIAGEVERTEXCOLORMASK_Red)
 				.IsChecked(this, &SFoliageEditMeshDisplayItem::IsVertexColorMaskChecked, FOLIAGEVERTEXCOLORMASK_Red)
+				.ToolTipText(LOCTEXT("VertexColorR_Tooltip", "Use the vertex color Red channel as a mask"))
 				[
 					SNew(STextBlock)
 					.Text(LOCTEXT("VertexColorR", "R"))
@@ -1153,6 +1191,7 @@ void SFoliageEditMeshDisplayItem::Construct(const FArguments& InArgs)
 				SNew(SCheckBox)
 				.OnCheckStateChanged(this, &SFoliageEditMeshDisplayItem::OnVertexColorMask, FOLIAGEVERTEXCOLORMASK_Green)
 				.IsChecked(this, &SFoliageEditMeshDisplayItem::IsVertexColorMaskChecked, FOLIAGEVERTEXCOLORMASK_Green)
+				.ToolTipText(LOCTEXT("VertexColorG_Tooltip", "Use the vertex color Green channel as a mask"))
 				[
 					SNew(STextBlock)
 					.Text(LOCTEXT("VertexColorG", "G"))
@@ -1168,6 +1207,7 @@ void SFoliageEditMeshDisplayItem::Construct(const FArguments& InArgs)
 				SNew(SCheckBox)
 				.OnCheckStateChanged(this, &SFoliageEditMeshDisplayItem::OnVertexColorMask, FOLIAGEVERTEXCOLORMASK_Blue)
 				.IsChecked(this, &SFoliageEditMeshDisplayItem::IsVertexColorMaskChecked, FOLIAGEVERTEXCOLORMASK_Blue)
+				.ToolTipText(LOCTEXT("VertexColorB_Tooltip", "Use the vertex color Blue channel as a mask"))
 				[
 					SNew(STextBlock)
 					.Text(LOCTEXT("VertexColorB", "B"))
@@ -1184,6 +1224,7 @@ void SFoliageEditMeshDisplayItem::Construct(const FArguments& InArgs)
 				SNew(SCheckBox)
 				.OnCheckStateChanged(this, &SFoliageEditMeshDisplayItem::OnVertexColorMask, FOLIAGEVERTEXCOLORMASK_Alpha)
 				.IsChecked(this, &SFoliageEditMeshDisplayItem::IsVertexColorMaskChecked, FOLIAGEVERTEXCOLORMASK_Alpha)
+				.ToolTipText(LOCTEXT("VertexColorA_Tooltip", "Use the vertex color Alpha channel as a mask"))
 				[
 					SNew(STextBlock)
 					.Text(LOCTEXT("VertexColorA", "A"))
@@ -1194,6 +1235,7 @@ void SFoliageEditMeshDisplayItem::Construct(const FArguments& InArgs)
 	TSharedRef<SHorizontalBox> VertexColorMaskThresholdBox =
 		SNew(SHorizontalBox)
 		.Visibility(this, &SFoliageEditMeshDisplayItem::IsVertexColorMaskThresholdVisible)
+		.ToolTipText(LOCTEXT("VertexColorMaskThreshold_Tooltip", "Specifies the threshold value above which the static mesh vertex color value must be, in order for foliage instances to be placed in a specific area."))
 
 		// Dummy Checkbox
 		+ SHorizontalBox::Slot()
@@ -1237,6 +1279,7 @@ void SFoliageEditMeshDisplayItem::Construct(const FArguments& InArgs)
 				SNew(SCheckBox)
 				.OnCheckStateChanged(this, &SFoliageEditMeshDisplayItem::OnVertexColorMaskInvert)
 				.IsChecked(this, &SFoliageEditMeshDisplayItem::IsVertexColorMaskInvertChecked)
+				.ToolTipText(LOCTEXT("VertexColorMaskThresholdInvert_Tooltip", "When unchecked, foliage instances will be placed only when the vertex color in the specified channel(s) is above the threshold amount. When checked, the vertex color must be less than the threshold amount."))
 			]
 
 			+ SHorizontalBox::Slot()
@@ -1245,6 +1288,7 @@ void SFoliageEditMeshDisplayItem::Construct(const FArguments& InArgs)
 			[
 				SNew(STextBlock)
 				.Text(LOCTEXT("VertexColorMaskThresholdInvert", "Invert"))
+				.ToolTipText(LOCTEXT("VertexColorMaskThresholdInvert_Tooltip", "When unchecked, foliage instances will be placed only when the vertex color in the specified channel(s) is above the threshold amount. When checked, the vertex color must be less than the threshold amount."))
 			]
 		];
 
diff --git a/Engine/Source/Editor/LandscapeEditor/Classes/LandscapeEditorObject.h b/Engine/Source/Editor/LandscapeEditor/Classes/LandscapeEditorObject.h
--- a/Engine/Source/Editor/LandscapeEditor/Classes/LandscapeEditorObject.h
+++ b/Engine/Source/Editor/LandscapeEditor/Classes/LandscapeEditorObject.h
@@ -102,13 +102,13 @@ namespace ELandscapeConvertMode
 	{
 		//Invalid = -1,
 
-		// Adds data to the edge of the landscape to get to a whole number of components
+		// Given the new component size, the edges of the landscape will be expanded as necessary until its overall size is a whole number of landscape components.
 		Expand = 0,
 
-		// Removes data from the edge of the landscape to get to a whole number of components
+		// Given the new component size, the edges of the landscape will be trimmed until its overall size is a whole number of landscape components.
 		Clip = 1,
 
-		// Keeps the landscape the same size with the same number of components, resamples the heightmap/weightmaps to the new resolution
+		// The landscape will have the same overall size in the world, and have the same number of components. Existing landscape geometry and layer data will be resampled to match the new resolution.
 		Resample = 2,
 	};
 }
@@ -380,7 +380,7 @@ class ULandscapeEditorObject : public UObject
 	UPROPERTY(Category="Change Component Size", EditAnywhere, NonTransactional, meta=(DisplayName="Number of Components", ShowForTools="ResizeLandscape"))
 	FIntPoint ResizeLandscape_ComponentCount;
 
-	// 
+	// Determines how the new component size will be applied to the existing landscape geometry.
 	UPROPERTY(Category="Change Component Size", EditAnywhere, NonTransactional, meta=(DisplayName="Resize Mode", ShowForTools="ResizeLandscape"))
 	TEnumAsByte<ELandscapeConvertMode::Type> ResizeLandscape_ConvertMode;
 
diff --git a/Engine/Source/Editor/LandscapeEditor/Private/LandscapeEditorDetailCustomization_MiscTools.cpp b/Engine/Source/Editor/LandscapeEditor/Private/LandscapeEditorDetailCustomization_MiscTools.cpp
--- a/Engine/Source/Editor/LandscapeEditor/Private/LandscapeEditorDetailCustomization_MiscTools.cpp
+++ b/Engine/Source/Editor/LandscapeEditor/Private/LandscapeEditorDetailCustomization_MiscTools.cpp
@@ -114,7 +114,7 @@ void FLandscapeEditorDetailCustomization_MiscTools::CustomizeDetails(IDetailLayo
 				SNew(STextBlock)
 				.Font(DetailBuilder.GetDetailFont())
 				.ShadowOffset(FVector2D::UnitVector)
-				.Text(LOCTEXT("Ramp.Hint", "Ctrl+Click to add ramp points, then press \"Add Ramp\"."))
+				.Text(LOCTEXT("Ramp.Hint", "Click to add ramp points, then press \"Add Ramp\"."))
 			]
 		];
 		ToolsCategory.AddCustomRow("Apply Ramp")
diff --git a/Engine/Source/Editor/LandscapeEditor/Private/LandscapeEditorDetailCustomization_ResizeLandscape.cpp b/Engine/Source/Editor/LandscapeEditor/Private/LandscapeEditorDetailCustomization_ResizeLandscape.cpp
--- a/Engine/Source/Editor/LandscapeEditor/Private/LandscapeEditorDetailCustomization_ResizeLandscape.cpp
+++ b/Engine/Source/Editor/LandscapeEditor/Private/LandscapeEditorDetailCustomization_ResizeLandscape.cpp
@@ -65,6 +65,7 @@ void FLandscapeEditorDetailCustomization_ResizeLandscape::CustomizeDetails(IDeta
 				SNew(STextBlock)
 				.Font(DetailBuilder.GetDetailFont())
 				.Text(LOCTEXT("Original", "Original"))
+				.ToolTipText(LOCTEXT("Original_Tip", "The properties of the landscape as it currently exists"))
 			]
 			+ SHorizontalBox::Slot()
 			.VAlign(VAlign_Center)
@@ -73,6 +74,7 @@ void FLandscapeEditorDetailCustomization_ResizeLandscape::CustomizeDetails(IDeta
 				SNew(STextBlock)
 				.Font(DetailBuilder.GetDetailFont())
 				.Text(LOCTEXT("New", "New"))
+				.ToolTipText(LOCTEXT("New_Tip", "The properties the landscape will have after the resize operation is completed"))
 			]
 		]
 	];
@@ -208,6 +210,7 @@ void FLandscapeEditorDetailCustomization_ResizeLandscape::CustomizeDetails(IDeta
 			SNew(STextBlock)
 			.Font(DetailBuilder.GetDetailFont())
 			.Text(LOCTEXT("Resolution", "Overall Resolution"))
+			.ToolTipText(LOCTEXT("Resolution_Tip", "Overall resolution of the entire landscape in vertices"))
 		]
 	]
 	.ValueContent()
@@ -247,6 +250,7 @@ void FLandscapeEditorDetailCustomization_ResizeLandscape::CustomizeDetails(IDeta
 			SNew(STextBlock)
 			.Font(DetailBuilder.GetDetailFont())
 			.Text(LOCTEXT("TotalComponents", "Total Components"))
+			.ToolTipText(LOCTEXT("TotalComponents_Tip", "The total number of components in the landscape"))
 		]
 	]
 	.ValueContent()
diff --git a/Engine/Source/Runtime/Engine/Classes/Foliage/FoliageType.h b/Engine/Source/Runtime/Engine/Classes/Foliage/FoliageType.h
--- a/Engine/Source/Runtime/Engine/Classes/Foliage/FoliageType.h
+++ b/Engine/Source/Runtime/Engine/Classes/Foliage/FoliageType.h
@@ -108,10 +108,15 @@ class UFoliageType : public UObject
 	UPROPERTY(EditAnywhere, Category=Painting)
 	uint32 VertexColorMaskInvert:1;
 
+	/* The maximum number of instances that can be placed in a single cluster component and rendered in a single draw call */
 	UPROPERTY(EditAnywhere, Category=Clustering)
 	int32 MaxInstancesPerCluster;
 
-	UPROPERTY(EditAnywhere, Category=Clustering)
+	/**
+	 * Determines the maximum size of a cluster of instances, rendered as a single component in a single draw call.
+	 * A smaller size will improve culling at the expense of more draw calls.
+	 */
+	UPROPERTY(EditAnywhere, Category = Clustering)
 	float MaxClusterRadius;
 
 	UPROPERTY()
@@ -159,10 +164,15 @@ class UFoliageType : public UObject
 	UPROPERTY()
 	uint32 ReapplyVertexColorMask:1;
 
+	/* The distance where instances will begin to fade out if using a PerInstanceFadeAmount material node. 0 disables. */
 	UPROPERTY(EditAnywhere, Category=Culling)
 	int32 StartCullDistance;
 
-	UPROPERTY(EditAnywhere, Category=Culling)
+	/**
+	 * The distance where instances will have completely faded out when using a PerInstanceFadeAmount material node. 0 disables. 
+	 * When the entire cluster is beyond this distance, the cluster is completely culled and not rendered at all.
+	 */
+	UPROPERTY(EditAnywhere, Category = Culling)
 	int32 EndCullDistance;
 
 	UPROPERTY()
@@ -216,7 +226,7 @@ class UFoliageType : public UObject
 	UPROPERTY(EditAnywhere, Category=Collision, meta=(HideObjectType=true))
 	struct FBodyInstance BodyInstance;
 
-	// Gets the mesh associated with this FoliageType
+	/* Gets the mesh associated with this FoliageType */
 	virtual UStaticMesh* GetStaticMesh() PURE_VIRTUAL(UFoliageType::GetStaticMesh, return nullptr; );
 };
 
