Editor - Foliage eiditor now has a much nicer drag/drop indicator.

diff --git a/Engine/Source/Editor/FoliageEdit/Private/SFoliageEdit.cpp b/Engine/Source/Editor/FoliageEdit/Private/SFoliageEdit.cpp
--- a/Engine/Source/Editor/FoliageEdit/Private/SFoliageEdit.cpp
+++ b/Engine/Source/Editor/FoliageEdit/Private/SFoliageEdit.cpp
@@ -2,6 +2,8 @@
 
 #include "UnrealEd.h"
 
+#include "SScaleBox.h"
+
 #include "SFoliageEdit.h"
 #include "FoliageEditActions.h"
 #include "FoliageEdMode.h"
@@ -27,21 +29,29 @@ public:
 	void Construct(const FArguments& InArgs)
 	{
 		FoliageEditPtr = InArgs._FoliageEditPtr;
+		bIsDragOn = false;
 
 		this->ChildSlot
 		[
-			InArgs._Content.Widget
+			SNew(SBorder)
+			.BorderImage(FEditorStyle::GetBrush("WhiteBrush"))
+			.BorderBackgroundColor(this, &SFoliageDragDropHandler::GetBackgroundColor)
+			.Padding(FMargin(30))
+			[
+				InArgs._Content.Widget
+			]
 		];
 	}
 
 	~SFoliageDragDropHandler() {}
 
-	FReply OnDrop(const FGeometry& MyGeometry, const FDragDropEvent& DragDropEvent)
+	FReply OnDrop(const FGeometry& MyGeometry, const FDragDropEvent& DragDropEvent) override
 	{
+		bIsDragOn = false;
+
 		TSharedPtr<SFoliageEdit> PinnedPtr = FoliageEditPtr.Pin();
 		if (PinnedPtr.IsValid())
 		{
-			PinnedPtr->DisableDragDropOverlay();
 			return PinnedPtr->OnDrop_ListView(MyGeometry, DragDropEvent);
 		}
 		return FReply::Handled();
@@ -49,26 +59,23 @@ public:
 
 	virtual void OnDragEnter(const FGeometry& MyGeometry, const FDragDropEvent& DragDropEvent) override
 	{
-		TSharedPtr<SFoliageEdit> PinnedPtr = FoliageEditPtr.Pin();
-		if (PinnedPtr.IsValid())
-		{
-			PinnedPtr->EnableDragDropOverlay();
-		}
-		SCompoundWidget::OnDragEnter(MyGeometry, DragDropEvent);
+		bIsDragOn = true;
 	}
 
 	virtual void OnDragLeave(const FDragDropEvent& DragDropEvent) override
 	{
-		TSharedPtr<SFoliageEdit> PinnedPtr = FoliageEditPtr.Pin();
-		if (PinnedPtr.IsValid())
-		{
-			PinnedPtr->DisableDragDropOverlay();
-		}
-		SCompoundWidget::OnDragLeave(DragDropEvent);
+		bIsDragOn = false;
+	}
+
+private:
+	FSlateColor GetBackgroundColor() const
+	{
+		return bIsDragOn ? FLinearColor(1.0f, 0.6f, 0.1f, 0.9f) : FLinearColor(0.1f, 0.1f, 0.1f, 0.9f);
 	}
 
 private:
 	TWeakPtr<SFoliageEdit> FoliageEditPtr;
+	bool bIsDragOn;
 };
 
 BEGIN_SLATE_FUNCTION_BUILD_OPTIMIZATION
@@ -86,247 +93,238 @@ void SFoliageEdit::Construct(const FArguments& InArgs)
 
 	AssetThumbnailPool = MakeShareable(new FAssetThumbnailPool(512, TAttribute<bool>::Create(TAttribute<bool>::FGetter::CreateSP(this, &SFoliageEdit::IsHovered))));
 
-	bOverlayOverride = false;
-
 	// Everything (or almost) uses this padding, change it to expand the padding.
 	FMargin StandardPadding(0.0f, 4.0f, 0.0f, 4.0f);
 
 	this->ChildSlot
 	[
-		SNew(SScrollBox)
-		+ SScrollBox::Slot()
-		.Padding(0.0f)
-		[
-			SNew(SVerticalBox)
+		SNew(SOverlay)
 
-			+ SVerticalBox::Slot()
-			.AutoHeight()
-			.Padding(StandardPadding)
+		+ SOverlay::Slot()
 			[
-				BuildToolBar()
-			]
-
-			+ SVerticalBox::Slot()
-			.AutoHeight()
-			[
-				SNew(SHorizontalBox)
-				.ToolTipText(LOCTEXT("BrushSize_Tooltip", "The size of the foliage brush"))
-
-				+ SHorizontalBox::Slot()
-				.FillWidth(1.0f)
-				.Padding(StandardPadding)
-				[
-					SNew(STextBlock)
-					.Visibility(this, &SFoliageEdit::GetVisibility_Radius)
-					.Text(LOCTEXT("BrushSize", "Brush Size"))
-				]
-				+ SHorizontalBox::Slot()
-				.FillWidth(2.0f)
-				.Padding(StandardPadding)
-				[
-					SNew(SSpinBox<float>)
-					.Visibility(this, &SFoliageEdit::GetVisibility_Radius)
-					.MinValue(0.0f)
-					.MaxValue(65536.0f)
-					.MaxSliderValue(8192.0f)
-					.SliderExponent(3.0f)
-					.Value(this, &SFoliageEdit::GetRadius)
-					.OnValueChanged(this, &SFoliageEdit::SetRadius)
-				]
-			]
-
-			+ SVerticalBox::Slot()
-			.AutoHeight()
+			SNew(SScrollBox)
+			+ SScrollBox::Slot()
+			.Padding(2.0f, 0)
 			[
-				SNew(SHorizontalBox)
-				.ToolTipText(LOCTEXT("PaintDensity_Tooltip", "The density of foliage to paint. This is a multiplier for the individual foliage type's density specifier."))
+				SNew(SVerticalBox)
 
-				+ SHorizontalBox::Slot()
-				.FillWidth(1.0f)
+				+ SVerticalBox::Slot()
+				.AutoHeight()
 				.Padding(StandardPadding)
 				[
-					SNew(STextBlock)
-					.Visibility(this, &SFoliageEdit::GetVisibility_PaintDensity)
-					.Text(LOCTEXT("PaintDensity", "Paint Density"))
+					BuildToolBar()
 				]
-				+ SHorizontalBox::Slot()
-				.FillWidth(2.0f)
-				.Padding(StandardPadding)
-				[
-					SNew(SSpinBox<float>)
-					.Visibility(this, &SFoliageEdit::GetVisibility_PaintDensity)
-					.MinValue(0.0f)
-					.MaxValue(2.0f)
-					.MaxSliderValue(1.0f)
-					.Value(this, &SFoliageEdit::GetPaintDensity)
-					.OnValueChanged(this, &SFoliageEdit::SetPaintDensity)
-				]
-			]
 
-			+ SVerticalBox::Slot()
-			.AutoHeight()
-			[
-				SNew(SHorizontalBox)
-				.ToolTipText(LOCTEXT("EraseDensity_Tooltip", "The density of foliage to leave behind when erasing with the Shift key held. 0 will remove all foliage."))
-
-				+ SHorizontalBox::Slot()
-				.FillWidth(1.0f)
-				.Padding(StandardPadding)
+				+ SVerticalBox::Slot()
+				.AutoHeight()
 				[
-					SNew(STextBlock)
-					.Visibility(this, &SFoliageEdit::GetVisibility_EraseDensity)
-					.Text(LOCTEXT("EraseDensity", "Erase Density"))
-				]
-				+ SHorizontalBox::Slot()
-				.FillWidth(2.0f)
-				.Padding(StandardPadding)
-				[
-					SNew(SSpinBox<float>)
-					.Visibility(this, &SFoliageEdit::GetVisibility_EraseDensity)
-					.MinValue(0.0f)
-					.MaxValue(2.0f)
-					.MaxSliderValue(1.0f)
-					.Value(this, &SFoliageEdit::GetEraseDensity)
-					.OnValueChanged(this, &SFoliageEdit::SetEraseDensity)
-				]
-			]
-
-			+ SVerticalBox::Slot()
-			.AutoHeight()
-			.Padding(StandardPadding)
-			[
-				SNew(STextBlock)
-				.Visibility(this, &SFoliageEdit::GetVisibility_Filters)
-				.Text(LOCTEXT("Filter", "Filter"))
-			]
-
-			+ SVerticalBox::Slot()
-			.AutoHeight()
-			.Padding(StandardPadding)
-			[
-				SNew(SWrapBox)
-				.UseAllottedWidth(true)
+					SNew(SHorizontalBox)
+					.ToolTipText(LOCTEXT("BrushSize_Tooltip", "The size of the foliage brush"))
 
-				+ SWrapBox::Slot()
-				.Padding(0, 0, 16, 0)
-				[
-					SNew(SCheckBox)
-					.Visibility(this, &SFoliageEdit::GetVisibility_Filters)
-					.OnCheckStateChanged(this, &SFoliageEdit::OnCheckStateChanged_Landscape)
-					.IsChecked(this, &SFoliageEdit::GetCheckState_Landscape)
-					.ToolTipText(LOCTEXT("FilterLandscape_Tooltip", "Place foliage on Landscape"))
+					+ SHorizontalBox::Slot()
+					.FillWidth(1.0f)
+					.Padding(StandardPadding)
 					[
 						SNew(STextBlock)
-						.Text(LOCTEXT("Landscape", "Landscape"))
+						.Visibility(this, &SFoliageEdit::GetVisibility_Radius)
+						.Text(LOCTEXT("BrushSize", "Brush Size"))
+					]
+					+ SHorizontalBox::Slot()
+					.FillWidth(2.0f)
+					.Padding(StandardPadding)
+					[
+						SNew(SSpinBox<float>)
+						.Visibility(this, &SFoliageEdit::GetVisibility_Radius)
+						.MinValue(0.0f)
+						.MaxValue(65536.0f)
+						.MaxSliderValue(8192.0f)
+						.SliderExponent(3.0f)
+						.Value(this, &SFoliageEdit::GetRadius)
+						.OnValueChanged(this, &SFoliageEdit::SetRadius)
 					]
 				]
 
-				+ SWrapBox::Slot()
-				.Padding(0, 0, 16, 0)
+				+ SVerticalBox::Slot()
+				.AutoHeight()
 				[
-					SNew(SCheckBox)
-					.Visibility(this, &SFoliageEdit::GetVisibility_Filters)
-					.OnCheckStateChanged(this, &SFoliageEdit::OnCheckStateChanged_StaticMesh)
-					.IsChecked(this, &SFoliageEdit::GetCheckState_StaticMesh)
-					.ToolTipText(LOCTEXT("FilterStaticMesh_Tooltip", "Place foliage on StaticMeshes"))
+					SNew(SHorizontalBox)
+					.ToolTipText(LOCTEXT("PaintDensity_Tooltip", "The density of foliage to paint. This is a multiplier for the individual foliage type's density specifier."))
+
+					+ SHorizontalBox::Slot()
+					.FillWidth(1.0f)
+					.Padding(StandardPadding)
 					[
 						SNew(STextBlock)
-						.Text(LOCTEXT("StaticMeshes", "Static Meshes"))
+						.Visibility(this, &SFoliageEdit::GetVisibility_PaintDensity)
+						.Text(LOCTEXT("PaintDensity", "Paint Density"))
+					]
+					+ SHorizontalBox::Slot()
+					.FillWidth(2.0f)
+					.Padding(StandardPadding)
+					[
+						SNew(SSpinBox<float>)
+						.Visibility(this, &SFoliageEdit::GetVisibility_PaintDensity)
+						.MinValue(0.0f)
+						.MaxValue(2.0f)
+						.MaxSliderValue(1.0f)
+						.Value(this, &SFoliageEdit::GetPaintDensity)
+						.OnValueChanged(this, &SFoliageEdit::SetPaintDensity)
 					]
 				]
 
-				+ SWrapBox::Slot()
-				.Padding(0, 0, 16, 0)
+				+ SVerticalBox::Slot()
+				.AutoHeight()
 				[
-					SNew(SCheckBox)
-					.Visibility(this, &SFoliageEdit::GetVisibility_Filters)
-					.OnCheckStateChanged(this, &SFoliageEdit::OnCheckStateChanged_BSP)
-					.IsChecked(this, &SFoliageEdit::GetCheckState_BSP)
-					.ToolTipText(LOCTEXT("FilterBSP_Tooltip", "Place foliage on BSP"))
+					SNew(SHorizontalBox)
+					.ToolTipText(LOCTEXT("EraseDensity_Tooltip", "The density of foliage to leave behind when erasing with the Shift key held. 0 will remove all foliage."))
+
+					+ SHorizontalBox::Slot()
+					.FillWidth(1.0f)
+					.Padding(StandardPadding)
 					[
 						SNew(STextBlock)
-						.Text(LOCTEXT("BSP", "BSP"))
+						.Visibility(this, &SFoliageEdit::GetVisibility_EraseDensity)
+						.Text(LOCTEXT("EraseDensity", "Erase Density"))
+					]
+					+ SHorizontalBox::Slot()
+					.FillWidth(2.0f)
+					.Padding(StandardPadding)
+					[
+						SNew(SSpinBox<float>)
+						.Visibility(this, &SFoliageEdit::GetVisibility_EraseDensity)
+						.MinValue(0.0f)
+						.MaxValue(2.0f)
+						.MaxSliderValue(1.0f)
+						.Value(this, &SFoliageEdit::GetEraseDensity)
+						.OnValueChanged(this, &SFoliageEdit::SetEraseDensity)
 					]
 				]
 
-				+ SWrapBox::Slot()
-				.Padding(0, 0, 16, 0)
+				+ SVerticalBox::Slot()
+				.AutoHeight()
+				.Padding(StandardPadding)
 				[
-					SNew(SCheckBox)
+					SNew(STextBlock)
 					.Visibility(this, &SFoliageEdit::GetVisibility_Filters)
-					.OnCheckStateChanged(this, &SFoliageEdit::OnCheckStateChanged_Translucent)
-					.IsChecked(this, &SFoliageEdit::GetCheckState_Translucent)
-					.ToolTipText(LOCTEXT("FilterTranslucent_Tooltip", "Place foliage on translucent surfaces"))
-					[
-						SNew(STextBlock)
-						.Text(LOCTEXT("Translucent", "Translucent"))
-					]
+					.Text(LOCTEXT("Filter", "Filter"))
 				]
-			]
 
-			+ SVerticalBox::Slot()
-			.AutoHeight()
-			[
-				SNew(SHeaderRow)
+				+ SVerticalBox::Slot()
+				.AutoHeight()
+				.Padding(StandardPadding)
+				[
+					SNew(SWrapBox)
+					.UseAllottedWidth(true)
 
-				+ SHeaderRow::Column(TEXT("Meshes"))
-				.DefaultLabel(LOCTEXT("Meshes", "Meshes"))
-			]
+					+ SWrapBox::Slot()
+					.Padding(0, 0, 16, 0)
+					[
+						SNew(SCheckBox)
+						.Visibility(this, &SFoliageEdit::GetVisibility_Filters)
+						.OnCheckStateChanged(this, &SFoliageEdit::OnCheckStateChanged_Landscape)
+						.IsChecked(this, &SFoliageEdit::GetCheckState_Landscape)
+						.ToolTipText(LOCTEXT("FilterLandscape_Tooltip", "Place foliage on Landscape"))
+						[
+							SNew(STextBlock)
+							.Text(LOCTEXT("Landscape", "Landscape"))
+						]
+					]
 
-			+ SVerticalBox::Slot()
-			.AutoHeight()
-			.Padding(StandardPadding)
-			[
+					+ SWrapBox::Slot()
+					.Padding(0, 0, 16, 0)
+					[
+						SNew(SCheckBox)
+						.Visibility(this, &SFoliageEdit::GetVisibility_Filters)
+						.OnCheckStateChanged(this, &SFoliageEdit::OnCheckStateChanged_StaticMesh)
+						.IsChecked(this, &SFoliageEdit::GetCheckState_StaticMesh)
+						.ToolTipText(LOCTEXT("FilterStaticMesh_Tooltip", "Place foliage on StaticMeshes"))
+						[
+							SNew(STextBlock)
+							.Text(LOCTEXT("StaticMeshes", "Static Meshes"))
+						]
+					]
 
-				SNew(SOverlay)
+					+ SWrapBox::Slot()
+					.Padding(0, 0, 16, 0)
+					[
+						SNew(SCheckBox)
+						.Visibility(this, &SFoliageEdit::GetVisibility_Filters)
+						.OnCheckStateChanged(this, &SFoliageEdit::OnCheckStateChanged_BSP)
+						.IsChecked(this, &SFoliageEdit::GetCheckState_BSP)
+						.ToolTipText(LOCTEXT("FilterBSP_Tooltip", "Place foliage on BSP"))
+						[
+							SNew(STextBlock)
+							.Text(LOCTEXT("BSP", "BSP"))
+						]
+					]
 
-				+ SOverlay::Slot()
-				[
-					SNew(SFoliageDragDropHandler)
-					.FoliageEditPtr(SharedThis(this))
+					+ SWrapBox::Slot()
+					.Padding(0, 0, 16, 0)
 					[
-						SAssignNew(ItemScrollBox, SScrollBox)
+						SNew(SCheckBox)
+						.Visibility(this, &SFoliageEdit::GetVisibility_Filters)
+						.OnCheckStateChanged(this, &SFoliageEdit::OnCheckStateChanged_Translucent)
+						.IsChecked(this, &SFoliageEdit::GetCheckState_Translucent)
+						.ToolTipText(LOCTEXT("FilterTranslucent_Tooltip", "Place foliage on translucent surfaces"))
+						[
+							SNew(STextBlock)
+							.Text(LOCTEXT("Translucent", "Translucent"))
+						]
 					]
 				]
 
-				+ SOverlay::Slot()
+				+ SVerticalBox::Slot()
+				.AutoHeight()
 				[
-					SNew(SHorizontalBox)
-					.Visibility(this, &SFoliageEdit::GetVisibility_EmptyList)
+					SNew(SHeaderRow)
 
-					+ SHorizontalBox::Slot()
-					.FillWidth(1.0f)
+					+ SHeaderRow::Column(TEXT("Meshes"))
+					.DefaultLabel(LOCTEXT("Meshes", "Meshes"))
+				]
+
+				+ SVerticalBox::Slot()
+				.AutoHeight()
+				.Padding(StandardPadding)
+				[
+					SNew(SBox)
+					.Visibility(this, &SFoliageEdit::GetVisibility_EmptyList)
+					.Padding(FMargin(15, 0))
+					.MinDesiredHeight(30)
 					[
-						SNew(SBorder)
-						.BorderImage(FEditorStyle::GetBrush("Menu.Background"))
-						.BorderBackgroundColor(FLinearColor(1, 1, 1, 0.8f))
-						.Content()
+						SNew(SScaleBox)
+						.Stretch(EStretch::ScaleToFit)
 						[
-							SNew(SHorizontalBox)
-
-							+ SHorizontalBox::Slot()
-							.FillWidth(1.0f)
-							[
-								SNew(SSpacer)
-							]
-
-							+ SHorizontalBox::Slot()
-							.AutoWidth()
-							.VAlign(VAlign_Center)
-							[
-								SNew(STextBlock)
-								.Visibility(this, &SFoliageEdit::GetVisibility_EmptyListText)
-								.Text(LOCTEXT("EmptyFoliageListMessage", "Drag static meshes from the content browser into this area."))
-							]
-
-							+ SHorizontalBox::Slot()
-							.FillWidth(1.0f)
-							[
-								SNew(SSpacer)
-							]
+							SNew(STextBlock)
+							.Text(LOCTEXT("Foliage_DropStatic", "+ Drop Static Meshes Here"))
+							.ToolTipText(LOCTEXT("Foliage_DropStatic_ToolTip", "Drag static meshes from the content browser into this area."))
 						]
 					]
 				]
+
+				+ SVerticalBox::Slot()
+				.AutoHeight()
+				[
+					SAssignNew(ItemScrollBox, SScrollBox)
+				]
+			]
+		]
+
+		// Static Mesh Drop Zone
+		+ SOverlay::Slot()
+		.HAlign(HAlign_Fill)
+		.VAlign(VAlign_Fill)
+		[
+			SNew(SFoliageDragDropHandler)
+			.Visibility(this, &SFoliageEdit::GetVisibility_FoliageDropTarget)
+			.FoliageEditPtr(SharedThis(this))
+			[
+				SNew(SScaleBox)
+				.Stretch(EStretch::ScaleToFit)
+				[
+					SNew(STextBlock)
+					.Text(LOCTEXT("Foliage_AddFoliageMesh", "+ Foliage Mesh"))
+					.ShadowOffset(FVector2D(1.f, 1.f))
+				]
 			]
 		]
 	];
@@ -599,27 +597,9 @@ float SFoliageEdit::GetEraseDensity() const
 	return FoliageEditMode->UISettings.GetUnpaintDensity();
 }
 
-void SFoliageEdit::OnDragEnter_ListView(const FGeometry& MyGeometry, const FDragDropEvent& DragDropEvent)
-{
-	bOverlayOverride = true;
-	SCompoundWidget::OnDragEnter(MyGeometry, DragDropEvent);
-}
-
-void SFoliageEdit::OnDragLeave_ListView(const FDragDropEvent& DragDropEvent)
-{
-	bOverlayOverride = false;
-	SCompoundWidget::OnDragLeave(DragDropEvent);
-}
-
-FReply SFoliageEdit::OnDragOver_ListView(const FGeometry& MyGeometry, const FDragDropEvent& DragDropEvent)
-{
-
-	return FReply::Handled();
-}
-
 EVisibility SFoliageEdit::GetVisibility_EmptyList() const
 {
-	if (DisplayItemList.Num() > 0 && !bOverlayOverride)
+	if (DisplayItemList.Num() > 0)
 	{
 		return EVisibility::Collapsed;
 	}
@@ -627,14 +607,21 @@ EVisibility SFoliageEdit::GetVisibility_EmptyList() const
 	return EVisibility::HitTestInvisible;
 }
 
-EVisibility SFoliageEdit::GetVisibility_EmptyListText() const
+EVisibility SFoliageEdit::GetVisibility_FoliageDropTarget() const
 {
-	if (DisplayItemList.Num() > 0)
+	if ( FSlateApplication::Get().IsDragDropping() )
 	{
-		return EVisibility::Collapsed;
+		TArray<FAssetData> DraggedAssets = AssetUtil::ExtractAssetDataFromDrag(FSlateApplication::Get().GetDragDroppingContent());
+		for ( const FAssetData& AssetData : DraggedAssets )
+		{
+			if ( AssetData.IsValid() && AssetData.GetClass()->IsChildOf(UStaticMesh::StaticClass()) )
+			{
+				return EVisibility::Visible;
+			}
+		}
 	}
-
-	return EVisibility::HitTestInvisible;
+	
+	return EVisibility::Hidden;
 }
 
 EVisibility SFoliageEdit::GetVisibility_NonEmptyList() const
@@ -762,16 +749,6 @@ EVisibility SFoliageEdit::GetVisibility_Filters() const
 	return EVisibility::Visible;
 }
 
-void SFoliageEdit::DisableDragDropOverlay(void)
-{
-	bOverlayOverride = false;
-}
-
-void SFoliageEdit::EnableDragDropOverlay(void)
-{
-	bOverlayOverride = true;
-}
-
 void SFoliageEdit::NotifyNewCurrentLevel()
 {
 	RefreshFullList();
diff --git a/Engine/Source/Editor/FoliageEdit/Private/SFoliageEdit.h b/Engine/Source/Editor/FoliageEdit/Private/SFoliageEdit.h
--- a/Engine/Source/Editor/FoliageEdit/Private/SFoliageEdit.h
+++ b/Engine/Source/Editor/FoliageEdit/Private/SFoliageEdit.h
@@ -44,12 +44,6 @@ public:
 	/** Gets FoliageEditMode. Used by the cluster details to notify changes */
 	class FEdModeFoliage* GetFoliageEditMode() const { return FoliageEditMode; }
 
-	/** Sets the overlay that appears over all items in the list to be visible. */
-	void DisableDragDropOverlay();
-
-	/** Sets the overlay that appears over all items in the list to be invisible. */
-	void EnableDragDropOverlay();
-
 private:
 
 	/** Clears all the tools selection by setting them to false. */
@@ -114,15 +108,6 @@ private:
 	/** Creates the list item widget that displays the instance settings. */
 	TSharedRef<ITableRow> MakeWidgetFromOption(TSharedPtr<FFoliageMeshUIInfo> InItem, const TSharedRef<STableViewBase>& OwnerTable);
 
-	/** Handles nothing? */
-	void OnDragEnter_ListView(const FGeometry& MyGeometry, const FDragDropEvent& DragDropEvent);
-
-	/** Handles nothing? */
-	void OnDragLeave_ListView(const FDragDropEvent& DragDropEvent);
-
-	/** @todo - Have something occur to help the user out in knowing something is actually going to occur (besides the mouse changing) */
-	FReply OnDragOver_ListView(const FGeometry& MyGeometry, const FDragDropEvent& DragDropEvent);
-
 	/** Sets the filter settings for if painting will occur on Landscapes. */
 	void OnCheckStateChanged_Landscape(ESlateCheckBoxState::Type InState);
 
@@ -150,8 +135,8 @@ private:
 	/** Checks if the text in the empty list overlay should appear. If the list is has items but the the drag and drop override is true, it will return EVisibility::Visible. */
 	EVisibility GetVisibility_EmptyList() const;
 
-	/** Checks if the empty list overlay and text block should appear. If the list is empty it will return EVisibility::Visible. */
-	EVisibility GetVisibility_EmptyListText() const;
+	/** Should the drop area be visible currently?  Happens when the user is dragging static meshes */
+	EVisibility GetVisibility_FoliageDropTarget() const;
 
 	/** Checks if the list should appear. */
 	EVisibility GetVisibility_NonEmptyList() const;
@@ -202,7 +187,4 @@ private:
 
 	/** List of items currently being displayed. */
 	TArray<TSharedRef<SFoliageEditMeshDisplayItem>> DisplayItemList;
-
-	/** Used to override the empty list overlay to be visible, not including the text. */
-	bool bOverlayOverride;
 };
diff --git a/Engine/Source/Editor/UnrealEd/Private/AssetSelection.cpp b/Engine/Source/Editor/UnrealEd/Private/AssetSelection.cpp
--- a/Engine/Source/Editor/UnrealEd/Private/AssetSelection.cpp
+++ b/Engine/Source/Editor/UnrealEd/Private/AssetSelection.cpp
@@ -419,11 +419,15 @@ static AActor* PrivateAddActor( UObject* Asset, UActorFactory* Factory, bool Sel
 
 namespace AssetUtil
 {
-	TArray<FAssetData> ExtractAssetDataFromDrag( const FDragDropEvent &DragDropEvent )
+	TArray<FAssetData> ExtractAssetDataFromDrag(const FDragDropEvent &DragDropEvent)
+	{
+		return ExtractAssetDataFromDrag(DragDropEvent.GetOperation());
+	}
+
+	TArray<FAssetData> ExtractAssetDataFromDrag(const TSharedPtr<FDragDropOperation>& Operation)
 	{
 		TArray<FAssetData> DroppedAssetData;
 
-		TSharedPtr<FDragDropOperation> Operation = DragDropEvent.GetOperation();
 		if (!Operation.IsValid())
 		{
 			return DroppedAssetData;
diff --git a/Engine/Source/Editor/UnrealEd/Public/AssetSelection.h b/Engine/Source/Editor/UnrealEd/Public/AssetSelection.h
--- a/Engine/Source/Editor/UnrealEd/Public/AssetSelection.h
+++ b/Engine/Source/Editor/UnrealEd/Public/AssetSelection.h
@@ -22,7 +22,10 @@ struct AssetMarshalDefs
 namespace AssetUtil
 {
 	/** Extract the names of the assets which are being dragged*/
-	UNREALED_API TArray<FAssetData>	ExtractAssetDataFromDrag( const class FDragDropEvent &DragDropEvent );
+	UNREALED_API TArray<FAssetData>	ExtractAssetDataFromDrag( const FDragDropEvent& DragDropEvent );
+
+	/** Extract the names of the assets which are being dragged*/
+	UNREALED_API TArray<FAssetData> ExtractAssetDataFromDrag( const TSharedPtr<FDragDropOperation>& Operation );
 
 	/**
 	 *  Given an array of asset names, loads the assets into an array of objects
