Fix foliage tool crashes when using the reapply tool with Z offset.

UE-4668

diff --git a/Engine/Source/Editor/FoliageEdit/Private/FoliageEdMode.cpp b/Engine/Source/Editor/FoliageEdit/Private/FoliageEdMode.cpp
--- a/Engine/Source/Editor/FoliageEdit/Private/FoliageEdMode.cpp
+++ b/Engine/Source/Editor/FoliageEdit/Private/FoliageEdMode.cpp
@@ -818,12 +818,13 @@ void FEdModeFoliage::ReapplyInstancesForBrush(UWorld* InWorld, AInstancedFoliage
 			// record that we've made changes to this instance already, so we don't touch it again.
 			Instance.Flags |= FOLIAGE_Readjusted;
 
+			// See if we need to update the location in the instance hash
 			bool bReapplyLocation = false;
+			FVector OldInstanceLocation = Instance.Location;
 
 			// remove any Z offset first, so the offset is reapplied to any new 
 			if (FMath::Abs(Instance.ZOffset) > KINDA_SMALL_NUMBER)
 			{
-				MeshInfo->InstanceHash->RemoveInstance(Instance.Location, InstanceIndex);
 				Instance.Location = Instance.GetInstanceWorldTransform().TransformPosition(FVector(0, 0, -Instance.ZOffset));
 				bReapplyLocation = true;
 			}
@@ -987,6 +988,11 @@ void FEdModeFoliage::ReapplyInstancesForBrush(UWorld* InWorld, AInstancedFoliage
 							(Settings->GroundSlope < 0.f && Hit.Normal.Z >= FMath::Cos(-PI * Settings->GroundSlope / 180.f) + SMALL_NUMBER))
 						{
 							InstancesToDelete.Add(InstanceIndex);
+							if (bReapplyLocation)
+							{
+								// restore the location so the hash removal will succeed
+								Instance.Location = OldInstanceLocation;
+							}
 							continue;
 						}
 					}
@@ -1009,6 +1015,11 @@ void FEdModeFoliage::ReapplyInstancesForBrush(UWorld* InWorld, AInstancedFoliage
 								if (HitWeight <= FMath::FRand())
 								{
 									InstancesToDelete.Add(InstanceIndex);
+									if (bReapplyLocation)
+									{
+										// restore the location so the hash removal will succeed
+										Instance.Location = OldInstanceLocation;
+									}
 									continue;
 								}
 							}
@@ -1029,6 +1040,11 @@ void FEdModeFoliage::ReapplyInstancesForBrush(UWorld* InWorld, AInstancedFoliage
 									if (!CheckVertexColor(Settings, VertexColor))
 									{
 										InstancesToDelete.Add(InstanceIndex);
+										if (bReapplyLocation)
+										{
+											// restore the location so the hash removal will succeed
+											Instance.Location = OldInstanceLocation;
+										}
 										continue;
 									}
 								}
@@ -1044,6 +1060,11 @@ void FEdModeFoliage::ReapplyInstancesForBrush(UWorld* InWorld, AInstancedFoliage
 				if (Instance.Location.Z < Settings->HeightMin || Instance.Location.Z > Settings->HeightMax)
 				{
 					InstancesToDelete.Add(InstanceIndex);
+					if (bReapplyLocation)
+					{
+						// restore the location so the hash removal will succeed
+						Instance.Location = OldInstanceLocation;
+					}
 					continue;
 				}
 			}
@@ -1052,11 +1073,13 @@ void FEdModeFoliage::ReapplyInstancesForBrush(UWorld* InWorld, AInstancedFoliage
 			{
 				// Reapply the Z offset in new local space
 				Instance.Location = Instance.GetInstanceWorldTransform().TransformPosition(FVector(0, 0, Instance.ZOffset));
+				bReapplyLocation = true;
 			}
 
-			// Re-add to the hash
+			// Update the hash
 			if (bReapplyLocation)
 			{
+				MeshInfo->InstanceHash->RemoveInstance(OldInstanceLocation, InstanceIndex);
 				MeshInfo->InstanceHash->InsertInstance(Instance.Location, InstanceIndex);
 			}
 
