Fixed up some low hanging fruit improving foliage tool placement performance with large numbers of instances

diff --git a/Engine/Source/Editor/FoliageEdit/Private/FoliageEdMode.cpp b/Engine/Source/Editor/FoliageEdit/Private/FoliageEdMode.cpp
--- a/Engine/Source/Editor/FoliageEdit/Private/FoliageEdMode.cpp
+++ b/Engine/Source/Editor/FoliageEdit/Private/FoliageEdMode.cpp
@@ -255,9 +255,9 @@ void FEdModeFoliage::Tick(FEditorViewportClient* ViewportClient, float DeltaTime
 	}
 }
 
-static bool FoliageTrace(UWorld* InWorld, FHitResult& OutHit, FVector InStart, FVector InEnd, FName InTraceTag, bool InbReturnFaceIndex = false)
+static bool FoliageTrace(UWorld* InWorld, AInstancedFoliageActor* IFA, FHitResult& OutHit, FVector InStart, FVector InEnd, FName InTraceTag, bool InbReturnFaceIndex = false)
 {
-	FCollisionQueryParams QueryParams(InTraceTag, true, AInstancedFoliageActor::GetInstancedFoliageActorForCurrentLevel(InWorld, false));
+	FCollisionQueryParams QueryParams(InTraceTag, true, IFA);
 	QueryParams.bReturnFaceIndex = InbReturnFaceIndex;
 
 	bool bResult = true;
@@ -307,8 +307,9 @@ void FEdModeFoliage::FoliageBrushTrace(FEditorViewportClient* ViewportClient, in
 
 			FHitResult Hit;
 			UWorld* World = ViewportClient->GetWorld();
+			AInstancedFoliageActor* IFA = AInstancedFoliageActor::GetInstancedFoliageActorForCurrentLevel(World, false);
 			static FName NAME_FoliageBrush = FName(TEXT("FoliageBrush"));
-			if (FoliageTrace(World, Hit, Start, End, NAME_FoliageBrush))
+			if (FoliageTrace(World, IFA, Hit, Start, End, NAME_FoliageBrush))
 			{
 				// Check filters
 				UPrimitiveComponent* PrimComp = Hit.Component.Get();
@@ -385,7 +386,7 @@ void FEdModeFoliage::GetRandomVectorInBrush(FVector& OutStart, FVector& OutEnd)
 // Number of buckets for layer weight histogram distribution.
 #define NUM_INSTANCE_BUCKETS 10
 
-bool CheckCollisionWithWorld(UFoliageType* Settings, const FFoliageInstance& Inst, const FVector& HitNormal, const FVector& HitLocation, UWorld* InWorld)
+bool CheckCollisionWithWorld(UFoliageType* Settings, const FFoliageInstance& Inst, const FVector& HitNormal, const FVector& HitLocation, UWorld* InWorld, AInstancedFoliageActor* InIFA)
 {
 	FMatrix InstTransform = Inst.GetInstanceWorldTransform().ToMatrixWithScale();
 	FVector LocalHit = InstTransform.InverseTransformPosition(HitLocation);
@@ -406,7 +407,7 @@ bool CheckCollisionWithWorld(UFoliageType* Settings, const FFoliageInstance& Ins
 				FVector SamplePos = InstTransform.TransformPosition(FVector(Settings->LowBoundOriginRadius.X, Settings->LowBoundOriginRadius.Y, 2.f) + LocalSamplePos[i]);
 				float WorldRadius = (Settings->LowBoundOriginRadius.Z + 2.f)*FMath::Max(Inst.DrawScale3D.X, Inst.DrawScale3D.Y);
 				FVector NormalVector = Settings->AlignToNormal ? HitNormal : FVector(0, 0, 1);
-				if (FoliageTrace(InWorld, Hit, SamplePos, SamplePos - NormalVector*WorldRadius, NAME_None))
+				if (FoliageTrace(InWorld, InIFA, Hit, SamplePos, SamplePos - NormalVector*WorldRadius, NAME_None))
 				{
 					if (LocalHit.Z - Inst.ZOffset < Settings->LowBoundOriginRadius.Z)
 					{
@@ -450,7 +451,7 @@ struct FPotentialInstance
 		, HitWeight(InHitWeight)
 	{}
 
-	bool PlaceInstance(UFoliageType* Settings, FFoliageInstance& Inst, UWorld* InWorld)
+	bool PlaceInstance(UFoliageType* Settings, FFoliageInstance& Inst, UWorld* InWorld, AInstancedFoliageActor* InIFA)
 	{
 		if (Settings->UniformScale)
 		{
@@ -494,7 +495,7 @@ struct FPotentialInstance
 
 		Inst.Base = HitComponent;
 
-		return CheckCollisionWithWorld(Settings, Inst, HitNormal, HitLocation, InWorld);
+		return CheckCollisionWithWorld(Settings, Inst, HitNormal, HitLocation, InWorld, InIFA);
 	}
 };
 
@@ -653,7 +654,7 @@ void FEdModeFoliage::AddInstancesForBrush(UWorld* InWorld, AInstancedFoliageActo
 
 			FHitResult Hit;
 			static FName NAME_AddInstancesForBrush = FName(TEXT("AddInstancesForBrush"));
-			if (FoliageTrace(InWorld, Hit, Start, End, NAME_AddInstancesForBrush, true))
+			if (FoliageTrace(InWorld, IFA, Hit, Start, End, NAME_AddInstancesForBrush, true))
 			{
 				// Check filters
 				UPrimitiveComponent* PrimComp = Hit.Component.Get();
@@ -733,7 +734,7 @@ void FEdModeFoliage::AddInstancesForBrush(UWorld* InWorld, AInstancedFoliageActo
 			for (int32 Idx = 0; Idx < AdditionalInstances; Idx++)
 			{
 				FFoliageInstance Inst;
-				if (PotentialInstances[Idx].PlaceInstance(Settings, Inst, InWorld))
+				if (PotentialInstances[Idx].PlaceInstance(Settings, Inst, InWorld, IFA))
 				{
 					MeshInfo.AddInstance(IFA, Settings, Inst);
 				}
@@ -972,7 +973,7 @@ void FEdModeFoliage::ReapplyInstancesForBrush(UWorld* InWorld, AInstancedFoliage
 				FVector ZAxis = Instance.Rotation.Quaternion().GetAxisZ();
 				FVector Start = Instance.Location + 16.f * ZAxis;
 				FVector End = Instance.Location - 16.f * ZAxis;
-				if (FoliageTrace(InWorld, Hit, Start, End, NAME_ReapplyInstancesForBrush, true))
+				if (FoliageTrace(InWorld, IFA, Hit, Start, End, NAME_ReapplyInstancesForBrush, true))
 				{
 					// Reapply the normal
 					if (bReapplyNormal)
@@ -1100,9 +1101,9 @@ void FEdModeFoliage::ReapplyInstancesForBrush(UWorld* InWorld, AInstancedFoliage
 				static const FName NAME_ReapplyInstancesForBrush = TEXT("ReapplyCollisionWithWorld");
 				FVector Start = Instance.Location + FVector(0.f, 0.f, 16.f);
 				FVector End = Instance.Location - FVector(0.f, 0.f, 16.f);
-				if (FoliageTrace(InWorld, Hit, Start, End, NAME_ReapplyInstancesForBrush))
+				if (FoliageTrace(InWorld, IFA, Hit, Start, End, NAME_ReapplyInstancesForBrush))
 				{
-					if (!CheckCollisionWithWorld(Settings, Instance, Hit.Normal, Hit.Location, InWorld))
+					if (!CheckCollisionWithWorld(Settings, Instance, Hit.Normal, Hit.Location, InWorld, IFA))
 					{
 						InstancesToDelete.Add(InstanceIndex);
 						continue;
@@ -1455,7 +1456,7 @@ void FEdModeFoliage::ApplyPaintBucket(AActor* Actor, bool bRemove)
 				for (int32 Idx = 0; Idx < InstancesToPlace.Num(); Idx++)
 				{
 					FFoliageInstance Inst;
-					if (InstancesToPlace[Idx].PlaceInstance(Settings, Inst, World))
+					if (InstancesToPlace[Idx].PlaceInstance(Settings, Inst, World, IFA))
 					{
 						MeshInfo.AddInstance(IFA, Settings, Inst);
 					}
@@ -1928,7 +1929,7 @@ bool FEdModeFoliage::InputKey(FEditorViewportClient* ViewportClient, FViewport*
 						FVector End = Instance.Location - FVector(0.f, 0.f, FOLIAGE_SNAP_TRACE);
 
 						FHitResult Hit;
-						if (FoliageTrace(GetWorld(), Hit, Start, End, FName("FoliageSnap")))
+						if (FoliageTrace(GetWorld(), IFA, Hit, Start, End, FName("FoliageSnap")))
 						{
 							// Check current level
 							if ((Hit.Component.IsValid() && Hit.Component.Get()->GetOutermost() == World->GetCurrentLevel()->GetOutermost()) ||
diff --git a/Engine/Source/Runtime/Engine/Private/InstancedFoliage.cpp b/Engine/Source/Runtime/Engine/Private/InstancedFoliage.cpp
--- a/Engine/Source/Runtime/Engine/Private/InstancedFoliage.cpp
+++ b/Engine/Source/Runtime/Engine/Private/InstancedFoliage.cpp
@@ -281,12 +281,15 @@ void FFoliageMeshInfo::AddInstance(AInstancedFoliageActor* InIFA, UFoliageType*
 	// Calculate transform for the instance
 	FTransform InstanceToWorld = InNewInstance.GetInstanceWorldTransform();
 
+	// Bounds of the new instance
+	FBoxSphereBounds InstanceWorldBounds = InSettings->GetStaticMesh()->GetBounds().TransformBy(InstanceToWorld);
+
 	if (BestCluster == nullptr)
 	{
 		BestClusterIndex = InstanceClusters.Num();
 		BestCluster = new(InstanceClusters)FFoliageInstanceCluster(
 			ConstructObject<UInstancedStaticMeshComponent>(UInstancedStaticMeshComponent::StaticClass(), InIFA, NAME_None, RF_Transactional),
-			InSettings->GetStaticMesh()->GetBounds().TransformBy(InstanceToWorld)
+			InstanceWorldBounds
 			);
 
 		BestCluster->ClusterComponent->Mobility = EComponentMobility::Static;
@@ -319,6 +322,7 @@ void FFoliageMeshInfo::AddInstance(AInstancedFoliageActor* InIFA, UFoliageType*
 	{
 		BestCluster->ClusterComponent->Modify();
 		BestCluster->ClusterComponent->InvalidateLightingCache();
+		BestCluster->Bounds = BestCluster->Bounds + InstanceWorldBounds;
 	}
 
 	BestCluster->InstanceIndices.Add(InstanceIndex);
@@ -334,9 +338,6 @@ void FFoliageMeshInfo::AddInstance(AInstancedFoliageActor* InIFA, UFoliageType*
 		BestCluster->ClusterComponent->SelectedInstances.Add(false);
 	}
 
-	// Update the bounds for the cluster
-	BestCluster->Bounds = BestCluster->ClusterComponent->CalcBounds(BestCluster->ClusterComponent->ComponentToWorld);
-
 	// Update PrimitiveComponent's culling distance taking into account the radius of the bounds, as
 	// it is based on the center of the component's bounds.
 	float CullDistance = InSettings->EndCullDistance > 0 ? (float)InSettings->EndCullDistance + BestCluster->Bounds.SphereRadius : 0.f;
