Foliage base component set to actor component as we don't actually need collision

diff --git a/Engine/Source/Editor/FoliageEdit/Private/FoliageEdMode.cpp b/Engine/Source/Editor/FoliageEdit/Private/FoliageEdMode.cpp
--- a/Engine/Source/Editor/FoliageEdit/Private/FoliageEdMode.cpp
+++ b/Engine/Source/Editor/FoliageEdit/Private/FoliageEdMode.cpp
@@ -754,7 +754,7 @@ void FEdModeFoliage::RemoveInstancesForBrush(AInstancedFoliageActor* IFA, FFolia
 		// Filter PotentialInstancesToRemove
 		for (int32 Idx = 0; Idx < PotentialInstancesToRemove.Num(); Idx++)
 		{
-			UPrimitiveComponent* Base = MeshInfo.Instances[PotentialInstancesToRemove[Idx]].Base;
+			UPrimitiveComponent* Base = Cast<UPrimitiveComponent> (MeshInfo.Instances[PotentialInstancesToRemove[Idx]].Base);
 			UMaterialInterface* Material = Base ? Base->GetMaterial(0) : nullptr;
 
 			// Check if instance is candidate for removal based on filter settings
diff --git a/Engine/Source/Runtime/Engine/Public/InstancedFoliage.h b/Engine/Source/Runtime/Engine/Public/InstancedFoliage.h
--- a/Engine/Source/Runtime/Engine/Public/InstancedFoliage.h
+++ b/Engine/Source/Runtime/Engine/Public/InstancedFoliage.h
@@ -52,7 +52,7 @@ struct FFoliageInstancePlacementInfo
  */
 struct FFoliageInstance : public FFoliageInstancePlacementInfo
 {
-	UPrimitiveComponent* Base;
+	UActorComponent* Base;
 
 	FFoliageInstance()
 	: Base(NULL)
