Slate - ESlateCheckBoxState has been renamed to ECheckBoxState.

diff --git a/Engine/Source/Developer/ProjectLauncher/Private/Widgets/Deploy/SProjectLauncherDeployFileServerSettings.cpp b/Engine/Source/Developer/ProjectLauncher/Private/Widgets/Deploy/SProjectLauncherDeployFileServerSettings.cpp
--- a/Engine/Source/Developer/ProjectLauncher/Private/Widgets/Deploy/SProjectLauncherDeployFileServerSettings.cpp
+++ b/Engine/Source/Developer/ProjectLauncher/Private/Widgets/Deploy/SProjectLauncherDeployFileServerSettings.cpp
@@ -83,18 +83,18 @@ void SProjectLauncherDeployFileServerSettings::Construct( const FArguments& InAr
 /* SProjectLauncherDeployFileServerSettings callbacks
  *****************************************************************************/
 
-void SProjectLauncherDeployFileServerSettings::HandleHideWindowCheckBoxCheckStateChanged( ESlateCheckBoxState::Type NewState )
+void SProjectLauncherDeployFileServerSettings::HandleHideWindowCheckBoxCheckStateChanged( ECheckBoxState NewState )
 {
 	ILauncherProfilePtr SelectedProfile = Model->GetSelectedProfile();
 
 	if (SelectedProfile.IsValid())
 	{
-		SelectedProfile->SetHideFileServerWindow(NewState == ESlateCheckBoxState::Checked);
+		SelectedProfile->SetHideFileServerWindow(NewState == ECheckBoxState::Checked);
 	}
 }
 
 
-ESlateCheckBoxState::Type SProjectLauncherDeployFileServerSettings::HandleHideWindowCheckBoxIsChecked( ) const
+ECheckBoxState SProjectLauncherDeployFileServerSettings::HandleHideWindowCheckBoxIsChecked( ) const
 {
 	ILauncherProfilePtr SelectedProfile = Model->GetSelectedProfile();
 
@@ -102,26 +102,26 @@ ESlateCheckBoxState::Type SProjectLauncherDeployFileServerSettings::HandleHideWi
 	{
 		if (SelectedProfile->IsFileServerHidden())
 		{
-			return ESlateCheckBoxState::Checked;
+			return ECheckBoxState::Checked;
 		}
 	}
 
-	return ESlateCheckBoxState::Unchecked;
+	return ECheckBoxState::Unchecked;
 }
 
 
-void SProjectLauncherDeployFileServerSettings::HandleStreamingServerCheckBoxCheckStateChanged( ESlateCheckBoxState::Type NewState )
+void SProjectLauncherDeployFileServerSettings::HandleStreamingServerCheckBoxCheckStateChanged( ECheckBoxState NewState )
 {
 	ILauncherProfilePtr SelectedProfile = Model->GetSelectedProfile();
 
 	if (SelectedProfile.IsValid())
 	{
-		SelectedProfile->SetStreamingFileServer(NewState == ESlateCheckBoxState::Checked);
+		SelectedProfile->SetStreamingFileServer(NewState == ECheckBoxState::Checked);
 	}
 }
 
 
-ESlateCheckBoxState::Type SProjectLauncherDeployFileServerSettings::HandleStreamingServerCheckBoxIsChecked( ) const
+ECheckBoxState SProjectLauncherDeployFileServerSettings::HandleStreamingServerCheckBoxIsChecked( ) const
 {
 	ILauncherProfilePtr SelectedProfile = Model->GetSelectedProfile();
 
@@ -129,11 +129,11 @@ ESlateCheckBoxState::Type SProjectLauncherDeployFileServerSettings::HandleStream
 	{
 		if (SelectedProfile->IsFileServerStreaming())
 		{
-			return ESlateCheckBoxState::Checked;
+			return ECheckBoxState::Checked;
 		}
 	}
 
-	return ESlateCheckBoxState::Unchecked;
+	return ECheckBoxState::Unchecked;
 }
 
 
diff --git a/Engine/Source/Developer/ProjectLauncher/Private/Widgets/Deploy/SProjectLauncherDeployFileServerSettings.h b/Engine/Source/Developer/ProjectLauncher/Private/Widgets/Deploy/SProjectLauncherDeployFileServerSettings.h
--- a/Engine/Source/Developer/ProjectLauncher/Private/Widgets/Deploy/SProjectLauncherDeployFileServerSettings.h
+++ b/Engine/Source/Developer/ProjectLauncher/Private/Widgets/Deploy/SProjectLauncherDeployFileServerSettings.h
@@ -28,16 +28,16 @@ public:
 private:
 
 	// Callback for check state changes of the 'Hide Window' check box.
-	void HandleHideWindowCheckBoxCheckStateChanged( ESlateCheckBoxState::Type NewState );
+	void HandleHideWindowCheckBoxCheckStateChanged( ECheckBoxState NewState );
 
 	// Callback for determining the checked state of the 'Hide Window' check box.
-	ESlateCheckBoxState::Type HandleHideWindowCheckBoxIsChecked( ) const;
+	ECheckBoxState HandleHideWindowCheckBoxIsChecked( ) const;
 
 	// Callback for check state changes of the 'Streaming Server' check box.
-	void HandleStreamingServerCheckBoxCheckStateChanged( ESlateCheckBoxState::Type NewState );
+	void HandleStreamingServerCheckBoxCheckStateChanged( ECheckBoxState NewState );
 
 	// Callback for determining the checked state of the 'Streaming Server' check box.
-	ESlateCheckBoxState::Type HandleStreamingServerCheckBoxIsChecked( ) const;
+	ECheckBoxState HandleStreamingServerCheckBoxIsChecked( ) const;
 
 private:
 
diff --git a/Engine/Source/Developer/ProjectLauncher/Private/Widgets/Deploy/SProjectLauncherDeployTargetListRow.h b/Engine/Source/Developer/ProjectLauncher/Private/Widgets/Deploy/SProjectLauncherDeployTargetListRow.h
--- a/Engine/Source/Developer/ProjectLauncher/Private/Widgets/Deploy/SProjectLauncherDeployTargetListRow.h
+++ b/Engine/Source/Developer/ProjectLauncher/Private/Widgets/Deploy/SProjectLauncherDeployTargetListRow.h
@@ -174,14 +174,14 @@ public:
 private:
 
 	// Callback for changing this row's check box state.
-	void HandleCheckBoxStateChanged(ESlateCheckBoxState::Type NewState)
+	void HandleCheckBoxStateChanged(ECheckBoxState NewState)
 	{
 		ILauncherDeviceGroupPtr ActiveGroup = DeviceGroup.Get();
 
 		if (ActiveGroup.IsValid() && DeviceProxy.IsValid() && DeviceProxy->HasVariant(SelectedVariant))
 		{
 			const FString& DeviceID = DeviceProxy->GetTargetDeviceId(SelectedVariant);
-			if (NewState == ESlateCheckBoxState::Checked)
+			if (NewState == ECheckBoxState::Checked)
 			{
 				ActiveGroup->AddDevice(DeviceID);
 				bDeviceInGroup = true;
@@ -195,14 +195,14 @@ private:
 	}
 
 	// Callback for determining this row's check box state.
-	ESlateCheckBoxState::Type HandleCheckBoxIsChecked() const
+	ECheckBoxState HandleCheckBoxIsChecked() const
 	{
 		if (IsEnabled())
 		{
-			return bDeviceInGroup ? ESlateCheckBoxState::Checked : ESlateCheckBoxState::Unchecked;
+			return bDeviceInGroup ? ECheckBoxState::Checked : ECheckBoxState::Unchecked;
 		}
 
-		return ESlateCheckBoxState::Unchecked;
+		return ECheckBoxState::Unchecked;
 	}
 
 	// Callback for getting the text of the current variant
diff --git a/Engine/Source/Developer/ProjectLauncher/Private/Widgets/Launch/SProjectLauncherLaunchRoleEditor.h b/Engine/Source/Developer/ProjectLauncher/Private/Widgets/Launch/SProjectLauncherLaunchRoleEditor.h
--- a/Engine/Source/Developer/ProjectLauncher/Private/Widgets/Launch/SProjectLauncherLaunchRoleEditor.h
+++ b/Engine/Source/Developer/ProjectLauncher/Private/Widgets/Launch/SProjectLauncherLaunchRoleEditor.h
@@ -197,18 +197,18 @@ private:
 	}
 
 	// Callback for checking the specified 'No VSync' check box.
-	void HandleVsyncCheckBoxCheckStateChanged( ESlateCheckBoxState::Type NewState )
+	void HandleVsyncCheckBoxCheckStateChanged( ECheckBoxState NewState )
 	{
 		ILauncherProfileLaunchRolePtr RolePtr = Role.Pin();
 
 		if (RolePtr.IsValid())
 		{
-			RolePtr->SetVsyncEnabled(NewState == ESlateCheckBoxState::Checked);
+			RolePtr->SetVsyncEnabled(NewState == ECheckBoxState::Checked);
 		}
 	}
 
 	// Callback for determining the checked state of the specified 'No VSync' check box.
-	ESlateCheckBoxState::Type HandleVsyncCheckBoxIsChecked( ) const
+	ECheckBoxState HandleVsyncCheckBoxIsChecked( ) const
 	{
 		ILauncherProfileLaunchRolePtr RolePtr = Role.Pin();
 
@@ -216,11 +216,11 @@ private:
 		{
 			if (RolePtr->IsVsyncEnabled())
 			{
-				return ESlateCheckBoxState::Checked;
+				return ECheckBoxState::Checked;
 			}
 		}
 
-		return ESlateCheckBoxState::Unchecked;
+		return ECheckBoxState::Unchecked;
 	}
 
 private:
diff --git a/Engine/Source/Developer/ProjectLauncher/Private/Widgets/SProjectLauncher.cpp b/Engine/Source/Developer/ProjectLauncher/Private/Widgets/SProjectLauncher.cpp
--- a/Engine/Source/Developer/ProjectLauncher/Private/Widgets/SProjectLauncher.cpp
+++ b/Engine/Source/Developer/ProjectLauncher/Private/Widgets/SProjectLauncher.cpp
@@ -243,14 +243,14 @@ void SProjectLauncher::FillWindowMenu( FMenuBuilder& MenuBuilder, TSharedRef<FWo
 /* SProjectLauncher callbacks
 *****************************************************************************/
 
-void SProjectLauncher::OnAdvancedChanged(const ESlateCheckBoxState::Type NewCheckedState)
+void SProjectLauncher::OnAdvancedChanged(const ECheckBoxState NewCheckedState)
 {
-	bAdvanced = (NewCheckedState == ESlateCheckBoxState::Checked);
+	bAdvanced = (NewCheckedState == ECheckBoxState::Checked);
 }
 
-ESlateCheckBoxState::Type SProjectLauncher::OnIsAdvanced() const
+ECheckBoxState SProjectLauncher::OnIsAdvanced() const
 {
-	return (bAdvanced) ? ESlateCheckBoxState::Checked : ESlateCheckBoxState::Unchecked;
+	return (bAdvanced) ? ECheckBoxState::Checked : ECheckBoxState::Unchecked;
 }
 
 const FSlateBrush* SProjectLauncher::GetAdvancedToggleBrush() const
diff --git a/Engine/Source/Developer/ProjectLauncher/Private/Widgets/SProjectLauncher.h b/Engine/Source/Developer/ProjectLauncher/Private/Widgets/SProjectLauncher.h
--- a/Engine/Source/Developer/ProjectLauncher/Private/Widgets/SProjectLauncher.h
+++ b/Engine/Source/Developer/ProjectLauncher/Private/Widgets/SProjectLauncher.h
@@ -62,10 +62,10 @@ protected:
 private:
 
 	// Callback for toggling the advanced mode.
-	void OnAdvancedChanged(const ESlateCheckBoxState::Type NewCheckedState);
+	void OnAdvancedChanged(const ECheckBoxState NewCheckedState);
 
 	// Callback to determine if we are in advanced mode.
-	ESlateCheckBoxState::Type OnIsAdvanced() const;
+	ECheckBoxState OnIsAdvanced() const;
 
 	// Get advanced toggle button brush.
 	const FSlateBrush* GetAdvancedToggleBrush() const;
diff --git a/Engine/Source/Developer/SessionFrontend/Private/Models/SessionBrowserOwnerFilter.h b/Engine/Source/Developer/SessionFrontend/Private/Models/SessionBrowserOwnerFilter.h
--- a/Engine/Source/Developer/SessionFrontend/Private/Models/SessionBrowserOwnerFilter.h
+++ b/Engine/Source/Developer/SessionFrontend/Private/Models/SessionBrowserOwnerFilter.h
@@ -40,9 +40,9 @@ public:
 	 *
 	 * @param CheckState The check box state.
 	 */
-	void EnableFromCheckState( ESlateCheckBoxState::Type CheckState )
+	void EnableFromCheckState( ECheckBoxState CheckState )
 	{
-		Enabled = (CheckState == ESlateCheckBoxState::Checked);
+		Enabled = (CheckState == ECheckBoxState::Checked);
 
 		OnStateChanged.ExecuteIfBound(OwnerName, Enabled);
 	}
@@ -62,9 +62,9 @@ public:
 	 *
 	 * @return Checked if the filter is enabled, Unchecked otherwise.
 	 */
-	ESlateCheckBoxState::Type GetCheckStateFromIsEnabled() const
+	ECheckBoxState GetCheckStateFromIsEnabled() const
 	{
-		return (Enabled ? ESlateCheckBoxState::Checked : ESlateCheckBoxState::Unchecked);
+		return (Enabled ? ECheckBoxState::Checked : ECheckBoxState::Unchecked);
 	}
 
 	/**
diff --git a/Engine/Source/Developer/SessionFrontend/Private/Models/SessionConsoleCategoryFilter.h b/Engine/Source/Developer/SessionFrontend/Private/Models/SessionConsoleCategoryFilter.h
--- a/Engine/Source/Developer/SessionFrontend/Private/Models/SessionConsoleCategoryFilter.h
+++ b/Engine/Source/Developer/SessionFrontend/Private/Models/SessionConsoleCategoryFilter.h
@@ -40,9 +40,9 @@ public:
 	 *
 	 * @param CheckState The check box state.
 	 */
-	void EnableFromCheckState( ESlateCheckBoxState::Type CheckState )
+	void EnableFromCheckState( ECheckBoxState CheckState )
 	{
-		Enabled = (CheckState == ESlateCheckBoxState::Checked);
+		Enabled = (CheckState == ECheckBoxState::Checked);
 
 		OnStateChanged.ExecuteIfBound(Category, Enabled);
 	}
@@ -62,9 +62,9 @@ public:
 	 *
 	 * @return Checked if the filter is enabled, Unchecked otherwise.
 	 */
-	ESlateCheckBoxState::Type GetCheckStateFromIsEnabled() const
+	ECheckBoxState GetCheckStateFromIsEnabled() const
 	{
-		return (Enabled ? ESlateCheckBoxState::Checked : ESlateCheckBoxState::Unchecked);
+		return (Enabled ? ECheckBoxState::Checked : ECheckBoxState::Unchecked);
 	}
 
 	/**
diff --git a/Engine/Source/Developer/SessionFrontend/Private/Widgets/Browser/SSessionBrowserInstanceListRow.h b/Engine/Source/Developer/SessionFrontend/Private/Widgets/Browser/SSessionBrowserInstanceListRow.h
--- a/Engine/Source/Developer/SessionFrontend/Private/Widgets/Browser/SSessionBrowserInstanceListRow.h
+++ b/Engine/Source/Developer/SessionFrontend/Private/Widgets/Browser/SSessionBrowserInstanceListRow.h
@@ -12,7 +12,7 @@
  * The first parameter is the engine instance that was checked or unchecked.
  * The second parameter is the new checked state.
  */
-DECLARE_DELEGATE_TwoParams(FOnSessionInstanceCheckStateChanged, const ISessionInstanceInfoPtr&, ESlateCheckBoxState::Type)
+DECLARE_DELEGATE_TwoParams(FOnSessionInstanceCheckStateChanged, const ISessionInstanceInfoPtr&, ECheckBoxState)
 
 
 /**
diff --git a/Engine/Source/Developer/SessionFrontend/Private/Widgets/Console/SSessionConsoleFilterBar.cpp b/Engine/Source/Developer/SessionFrontend/Private/Widgets/Console/SSessionConsoleFilterBar.cpp
--- a/Engine/Source/Developer/SessionFrontend/Private/Widgets/Console/SSessionConsoleFilterBar.cpp
+++ b/Engine/Source/Developer/SessionFrontend/Private/Widgets/Console/SSessionConsoleFilterBar.cpp
@@ -213,7 +213,7 @@ void SSessionConsoleFilterBar::HandleFilterStringTextChanged( const FText& NewTe
 }
 
 
-void SSessionConsoleFilterBar::HandleHighlightOnlyCheckBoxCheckStateChanged( ESlateCheckBoxState::Type CheckedState )
+void SSessionConsoleFilterBar::HandleHighlightOnlyCheckBoxCheckStateChanged( ECheckBoxState CheckedState )
 {
 	OnFilterChanged.ExecuteIfBound();
 }
diff --git a/Engine/Source/Developer/SessionFrontend/Private/Widgets/Console/SSessionConsoleFilterBar.h b/Engine/Source/Developer/SessionFrontend/Private/Widgets/Console/SSessionConsoleFilterBar.h
--- a/Engine/Source/Developer/SessionFrontend/Private/Widgets/Console/SSessionConsoleFilterBar.h
+++ b/Engine/Source/Developer/SessionFrontend/Private/Widgets/Console/SSessionConsoleFilterBar.h
@@ -80,7 +80,7 @@ private:
 	void HandleFilterStringTextChanged( const FText& NewText );
 
 	/** Callback for changing the checked state of the 'Filter' check box. */
-	void HandleHighlightOnlyCheckBoxCheckStateChanged( ESlateCheckBoxState::Type CheckedState );
+	void HandleHighlightOnlyCheckBoxCheckStateChanged( ECheckBoxState CheckedState );
 
 	/** Callback for generating a row widget for the verbosity filter list. */
 	TSharedRef<ITableRow> HandleVerbosityFilterGenerateRow( FSessionConsoleVerbosityFilterPtr Filter, const TSharedRef<STableViewBase>& OwnerTable );
diff --git a/Engine/Source/Developer/SlateReflector/Private/Widgets/SAtlasVisualizer.cpp b/Engine/Source/Developer/SlateReflector/Private/Widgets/SAtlasVisualizer.cpp
--- a/Engine/Source/Developer/SlateReflector/Private/Widgets/SAtlasVisualizer.cpp
+++ b/Engine/Source/Developer/SlateReflector/Private/Widgets/SAtlasVisualizer.cpp
@@ -404,11 +404,11 @@ FText SAtlasVisualizer::GetZoomLevelPercentText() const
 	return FText::GetEmpty();
 }
 
-void SAtlasVisualizer::OnFitToWindowStateChanged( ESlateCheckBoxState::Type NewState )
+void SAtlasVisualizer::OnFitToWindowStateChanged( ECheckBoxState NewState )
 {
 	if( ScrollPanel.IsValid() )
 	{
-		if( NewState == ESlateCheckBoxState::Checked )
+		if( NewState == ECheckBoxState::Checked )
 		{
 			ScrollPanel->FitToWindow();
 		}
@@ -419,14 +419,14 @@ void SAtlasVisualizer::OnFitToWindowStateChanged( ESlateCheckBoxState::Type NewS
 	}
 }
 
-ESlateCheckBoxState::Type SAtlasVisualizer::OnGetFitToWindowState() const
+ECheckBoxState SAtlasVisualizer::OnGetFitToWindowState() const
 {
 	if( ScrollPanel.IsValid() )
 	{
-		return (ScrollPanel->IsFitToWindow()) ? ESlateCheckBoxState::Checked : ESlateCheckBoxState::Unchecked;
+		return (ScrollPanel->IsFitToWindow()) ? ECheckBoxState::Checked : ECheckBoxState::Unchecked;
 	}
 
-	return ESlateCheckBoxState::Undetermined;
+	return ECheckBoxState::Undetermined;
 }
 
 FReply SAtlasVisualizer::OnActualSizeClicked()
@@ -439,14 +439,14 @@ FReply SAtlasVisualizer::OnActualSizeClicked()
 	return FReply::Handled();
 }
 
-void SAtlasVisualizer::OnDisplayCheckerboardStateChanged( ESlateCheckBoxState::Type NewState )
+void SAtlasVisualizer::OnDisplayCheckerboardStateChanged( ECheckBoxState NewState )
 {
-	bDisplayCheckerboard = NewState == ESlateCheckBoxState::Checked;
+	bDisplayCheckerboard = NewState == ECheckBoxState::Checked;
 }
 
-ESlateCheckBoxState::Type SAtlasVisualizer::OnGetCheckerboardState() const
+ECheckBoxState SAtlasVisualizer::OnGetCheckerboardState() const
 {
-	return bDisplayCheckerboard ? ESlateCheckBoxState::Checked : ESlateCheckBoxState::Unchecked;
+	return bDisplayCheckerboard ? ECheckBoxState::Checked : ECheckBoxState::Unchecked;
 }
 
 EVisibility SAtlasVisualizer::OnGetCheckerboardVisibility() const
diff --git a/Engine/Source/Developer/SlateReflector/Private/Widgets/SAtlasVisualizer.h b/Engine/Source/Developer/SlateReflector/Private/Widgets/SAtlasVisualizer.h
--- a/Engine/Source/Developer/SlateReflector/Private/Widgets/SAtlasVisualizer.h
+++ b/Engine/Source/Developer/SlateReflector/Private/Widgets/SAtlasVisualizer.h
@@ -26,11 +26,11 @@ public:
 
 protected:
 	FText GetZoomLevelPercentText() const;
-	void OnFitToWindowStateChanged( ESlateCheckBoxState::Type NewState );
-	ESlateCheckBoxState::Type OnGetFitToWindowState() const;
+	void OnFitToWindowStateChanged( ECheckBoxState NewState );
+	ECheckBoxState OnGetFitToWindowState() const;
 	FReply OnActualSizeClicked();
-	void OnDisplayCheckerboardStateChanged( ESlateCheckBoxState::Type NewState );
-	ESlateCheckBoxState::Type OnGetCheckerboardState() const;
+	void OnDisplayCheckerboardStateChanged( ECheckBoxState NewState );
+	ECheckBoxState OnGetCheckerboardState() const;
 	EVisibility OnGetCheckerboardVisibility() const;
 	void OnComboOpening();
 	FText OnGetSelectedItemText() const;
diff --git a/Engine/Source/Developer/SlateReflector/Private/Widgets/SWidgetReflector.cpp b/Engine/Source/Developer/SlateReflector/Private/Widgets/SWidgetReflector.cpp
--- a/Engine/Source/Developer/SlateReflector/Private/Widgets/SWidgetReflector.cpp
+++ b/Engine/Source/Developer/SlateReflector/Private/Widgets/SWidgetReflector.cpp
@@ -189,11 +189,11 @@ void SWidgetReflector::Construct( const FArguments& InArgs )
 						.Style( FCoreStyle::Get(), "ToggleButtonCheckbox" )
 						.IsChecked_Lambda([]()
 						{
-							return bFoldTick == 0 ? ESlateCheckBoxState::Unchecked : ESlateCheckBoxState::Checked;
+							return bFoldTick == 0 ? ECheckBoxState::Unchecked : ECheckBoxState::Checked;
 						})
-						.OnCheckStateChanged_Lambda([]( const ESlateCheckBoxState::Type NewState )
+						.OnCheckStateChanged_Lambda([]( const ECheckBoxState NewState )
 						{
-							bFoldTick = (NewState == ESlateCheckBoxState::Checked) ? 1 : 0;
+							bFoldTick = (NewState == ECheckBoxState::Checked) ? 1 : 0;
 						})
 						[
 							SNew(SBox)
@@ -458,9 +458,9 @@ FReply SWidgetReflector::HandleDisplayFontAtlases()
 	return FReply::Handled();
 }
 
-void SWidgetReflector::HandleFocusCheckBoxCheckedStateChanged( ESlateCheckBoxState::Type NewValue )
+void SWidgetReflector::HandleFocusCheckBoxCheckedStateChanged( ECheckBoxState NewValue )
 {
-	bShowFocus = NewValue != ESlateCheckBoxState::Unchecked;
+	bShowFocus = NewValue != ECheckBoxState::Unchecked;
 
 	if (bShowFocus)
 	{
diff --git a/Engine/Source/Developer/SlateReflector/Private/Widgets/SWidgetReflector.h b/Engine/Source/Developer/SlateReflector/Private/Widgets/SWidgetReflector.h
--- a/Engine/Source/Developer/SlateReflector/Private/Widgets/SWidgetReflector.h
+++ b/Engine/Source/Developer/SlateReflector/Private/Widgets/SWidgetReflector.h
@@ -135,12 +135,12 @@ private:
 	}
 
 	/** Callback for checked state changes of the focus check box. */
-	void HandleFocusCheckBoxCheckedStateChanged( ESlateCheckBoxState::Type NewValue );
+	void HandleFocusCheckBoxCheckedStateChanged( ECheckBoxState NewValue );
 
 	/** Callback for getting the checked state of the focus check box. */
-	ESlateCheckBoxState::Type HandleFocusCheckBoxIsChecked() const
+	ECheckBoxState HandleFocusCheckBoxIsChecked() const
 	{
-		return bShowFocus ? ESlateCheckBoxState::Checked : ESlateCheckBoxState::Unchecked;
+		return bShowFocus ? ECheckBoxState::Checked : ECheckBoxState::Unchecked;
 	}
 
 	/** Callback for getting the text of the frame rate text block. */
diff --git a/Engine/Source/Editor/BspMode/Private/SBspPalette.cpp b/Engine/Source/Editor/BspMode/Private/SBspPalette.cpp
--- a/Engine/Source/Editor/BspMode/Private/SBspPalette.cpp
+++ b/Engine/Source/Editor/BspMode/Private/SBspPalette.cpp
@@ -217,24 +217,24 @@ FReply SBspPalette::OnDraggingListViewWidget(const FGeometry& MyGeometry, const
 	return FReply::Unhandled();
 }
 
-void SBspPalette::OnAdditiveModeButtonClicked(ESlateCheckBoxState::Type CheckType)
+void SBspPalette::OnAdditiveModeButtonClicked(ECheckBoxState CheckType)
 {
-	bIsAdditive = CheckType == ESlateCheckBoxState::Checked;
+	bIsAdditive = CheckType == ECheckBoxState::Checked;
 }
 
-void SBspPalette::OnSubtractiveModeButtonClicked(ESlateCheckBoxState::Type CheckType)
+void SBspPalette::OnSubtractiveModeButtonClicked(ECheckBoxState CheckType)
 {
-	bIsAdditive = CheckType != ESlateCheckBoxState::Checked;
+	bIsAdditive = CheckType != ECheckBoxState::Checked;
 }
 
-ESlateCheckBoxState::Type SBspPalette::IsAdditiveModeChecked() const
+ECheckBoxState SBspPalette::IsAdditiveModeChecked() const
 {
-	return bIsAdditive ? ESlateCheckBoxState::Checked : ESlateCheckBoxState::Unchecked;
+	return bIsAdditive ? ECheckBoxState::Checked : ECheckBoxState::Unchecked;
 }
 
-ESlateCheckBoxState::Type SBspPalette::IsSubtractiveModeChecked() const
+ECheckBoxState SBspPalette::IsSubtractiveModeChecked() const
 {
-	return !bIsAdditive ? ESlateCheckBoxState::Checked : ESlateCheckBoxState::Unchecked;
+	return !bIsAdditive ? ECheckBoxState::Checked : ECheckBoxState::Unchecked;
 }
 
 const FSlateBrush* SBspPalette::GetAdditiveModeImage() const
diff --git a/Engine/Source/Editor/BspMode/Private/SBspPalette.h b/Engine/Source/Editor/BspMode/Private/SBspPalette.h
--- a/Engine/Source/Editor/BspMode/Private/SBspPalette.h
+++ b/Engine/Source/Editor/BspMode/Private/SBspPalette.h
@@ -22,16 +22,16 @@ private:
 	FReply OnDraggingListViewWidget(const FGeometry& MyGeometry, const FPointerEvent& MouseEvent);
 
 	/** Select additive mode */
-	void OnAdditiveModeButtonClicked(ESlateCheckBoxState::Type CheckType);
+	void OnAdditiveModeButtonClicked(ECheckBoxState CheckType);
 
 	/** Select subtractive mode */
-	void OnSubtractiveModeButtonClicked(ESlateCheckBoxState::Type CheckType);
+	void OnSubtractiveModeButtonClicked(ECheckBoxState CheckType);
 
 	/** @return the check state of the additive mode checkbox */
-	ESlateCheckBoxState::Type IsAdditiveModeChecked() const;
+	ECheckBoxState IsAdditiveModeChecked() const;
 
 	/** @return the check state of the subtractive mode checkbox */
-	ESlateCheckBoxState::Type IsSubtractiveModeChecked() const;
+	ECheckBoxState IsSubtractiveModeChecked() const;
 
 	/** Get the image for additive mode */
 	const FSlateBrush* GetAdditiveModeImage() const;
diff --git a/Engine/Source/Editor/ContentBrowser/Private/SAssetPicker.cpp b/Engine/Source/Editor/ContentBrowser/Private/SAssetPicker.cpp
--- a/Engine/Source/Editor/ContentBrowser/Private/SAssetPicker.cpp
+++ b/Engine/Source/Editor/ContentBrowser/Private/SAssetPicker.cpp
@@ -414,14 +414,14 @@ FText SAssetPicker::GetShowOtherDevelopersToolTip() const
 	}
 }
 
-void SAssetPicker::HandleShowOtherDevelopersCheckStateChanged( ESlateCheckBoxState::Type InCheckboxState )
+void SAssetPicker::HandleShowOtherDevelopersCheckStateChanged( ECheckBoxState InCheckboxState )
 {
-	OtherDevelopersFilter->SetShowOtherDeveloperAssets( InCheckboxState == ESlateCheckBoxState::Checked );
+	OtherDevelopersFilter->SetShowOtherDeveloperAssets( InCheckboxState == ECheckBoxState::Checked );
 }
 
-ESlateCheckBoxState::Type SAssetPicker::GetShowOtherDevelopersCheckState() const
+ECheckBoxState SAssetPicker::GetShowOtherDevelopersCheckState() const
 {
-	return OtherDevelopersFilter->GetShowOtherDeveloperAssets() ? ESlateCheckBoxState::Checked : ESlateCheckBoxState::Unchecked;
+	return OtherDevelopersFilter->GetShowOtherDeveloperAssets() ? ECheckBoxState::Checked : ECheckBoxState::Unchecked;
 }
 
 void SAssetPicker::OnRenameRequested() const
diff --git a/Engine/Source/Editor/ContentBrowser/Private/SAssetPicker.h b/Engine/Source/Editor/ContentBrowser/Private/SAssetPicker.h
--- a/Engine/Source/Editor/ContentBrowser/Private/SAssetPicker.h
+++ b/Engine/Source/Editor/ContentBrowser/Private/SAssetPicker.h
@@ -65,10 +65,10 @@ private:
 	FText GetShowOtherDevelopersToolTip() const;
 
 	/** Toggles the filter for showing other developers assets */
-	void HandleShowOtherDevelopersCheckStateChanged( const ESlateCheckBoxState::Type InCheckboxState );
+	void HandleShowOtherDevelopersCheckStateChanged( const ECheckBoxState InCheckboxState );
 
 	/** Gets if showing other developers assets */
-	ESlateCheckBoxState::Type GetShowOtherDevelopersCheckState() const;
+	ECheckBoxState GetShowOtherDevelopersCheckState() const;
 
 	/** Called upon the Rename UICommand being executed, sends rename request to the asset view */
 	void OnRenameRequested() const;
diff --git a/Engine/Source/Editor/ContentBrowser/Private/SFilterList.cpp b/Engine/Source/Editor/ContentBrowser/Private/SFilterList.cpp
--- a/Engine/Source/Editor/ContentBrowser/Private/SFilterList.cpp
+++ b/Engine/Source/Editor/ContentBrowser/Private/SFilterList.cpp
@@ -187,9 +187,9 @@ public:
 
 private:
 	/** Handler for when the filter checkbox is clicked */
-	void FilterToggled(ESlateCheckBoxState::Type NewState)
+	void FilterToggled(ECheckBoxState NewState)
 	{
-		bEnabled = NewState == ESlateCheckBoxState::Checked;
+		bEnabled = NewState == ECheckBoxState::Checked;
 		OnFilterChanged.ExecuteIfBound();
 	}
 
@@ -283,15 +283,15 @@ private:
 	}
 
 	/** Handler to determine the "checked" state of the filter checkbox */
-	ESlateCheckBoxState::Type IsChecked() const
+	ECheckBoxState IsChecked() const
 	{
-		return bEnabled ? ESlateCheckBoxState::Checked : ESlateCheckBoxState::Unchecked;
+		return bEnabled ? ECheckBoxState::Checked : ECheckBoxState::Unchecked;
 	}
 
 	/** Handler to determine the color of the checkbox when it is checked */
 	FSlateColor GetFilterForegroundColor() const
 	{
-		return IsChecked() ? FilterColor : FLinearColor::White;
+		return IsChecked() == ECheckBoxState::Checked ? FilterColor : FLinearColor::White;
 	}
 
 	/** Handler to determine the padding of the checkbox text when it is pressed */
diff --git a/Engine/Source/Editor/DetailCustomizations/Private/BodyInstanceCustomization.cpp b/Engine/Source/Editor/DetailCustomizations/Private/BodyInstanceCustomization.cpp
--- a/Engine/Source/Editor/DetailCustomizations/Private/BodyInstanceCustomization.cpp
+++ b/Engine/Source/Editor/DetailCustomizations/Private/BodyInstanceCustomization.cpp
@@ -787,7 +787,7 @@ EVisibility FBodyInstanceCustomization::ShouldShowResetToDefaultResponse(int32 V
 	{
 		const ECollisionResponse DefaultResponse = FCollisionResponseContainer::GetDefaultResponseContainer().GetResponse(ValidCollisionChannels[ValidIndex].CollisionChannel);
 
-		if (IsCollisionChannelChecked(ValidIndex, DefaultResponse) != ESlateCheckBoxState::Checked)
+		if (IsCollisionChannelChecked(ValidIndex, DefaultResponse) != ECheckBoxState::Checked)
 		{
 			return EVisibility::Visible;
 		}
@@ -838,7 +838,7 @@ FString FBodyInstanceCustomization::GetCollisionProfileComboBoxToolTip() const
 	return TEXT("Multiple Values");
 }
 
-void FBodyInstanceCustomization::OnCollisionChannelChanged(ESlateCheckBoxState::Type InNewValue, int32 ValidIndex, ECollisionResponse InCollisionResponse)
+void FBodyInstanceCustomization::OnCollisionChannelChanged(ECheckBoxState InNewValue, int32 ValidIndex, ECollisionResponse InCollisionResponse)
 {
 	if ( ValidCollisionChannels.IsValidIndex(ValidIndex) )
 	{
@@ -862,7 +862,7 @@ void FBodyInstanceCustomization::SetResponse(int32 ValidIndex, ECollisionRespons
 	CollisionResponsesHandle->NotifyPostChange();
 }
 
-ESlateCheckBoxState::Type FBodyInstanceCustomization::IsCollisionChannelChecked( int32 ValidIndex, ECollisionResponse InCollisionResponse) const
+ECheckBoxState FBodyInstanceCustomization::IsCollisionChannelChecked( int32 ValidIndex, ECollisionResponse InCollisionResponse) const
 {
 	TArray<uint8> CollisionResponses;
 
@@ -879,35 +879,35 @@ ESlateCheckBoxState::Type FBodyInstanceCustomization::IsCollisionChannelChecked(
 		{
 			if (CollisionResponses[0] == InCollisionResponse)
 			{
-				return ESlateCheckBoxState::Checked;
+				return ECheckBoxState::Checked;
 			}
 			else
 			{
-				return ESlateCheckBoxState::Unchecked;
+				return ECheckBoxState::Unchecked;
 			}
 		}
 		else if (CollisionResponses.Contains(InCollisionResponse))
 		{
-			return ESlateCheckBoxState::Undetermined;
+			return ECheckBoxState::Undetermined;
 		}
 
 		// if it didn't contain and it's not found, return Unchecked
-		return ESlateCheckBoxState::Unchecked;
+		return ECheckBoxState::Unchecked;
 	}
 
-	return ESlateCheckBoxState::Undetermined;
+	return ECheckBoxState::Undetermined;
 }
 
-void FBodyInstanceCustomization::OnAllCollisionChannelChanged(ESlateCheckBoxState::Type InNewValue, ECollisionResponse InCollisionResponse)
+void FBodyInstanceCustomization::OnAllCollisionChannelChanged(ECheckBoxState InNewValue, ECollisionResponse InCollisionResponse)
 {
 	FCollisionResponseContainer NewContainer;
 	NewContainer.SetAllChannels(InCollisionResponse);
 	SetCollisionResponseContainer(NewContainer);
 }
 
-ESlateCheckBoxState::Type FBodyInstanceCustomization::IsAllCollisionChannelChecked(ECollisionResponse InCollisionResponse) const
+ECheckBoxState FBodyInstanceCustomization::IsAllCollisionChannelChecked(ECollisionResponse InCollisionResponse) const
 {
-	ESlateCheckBoxState::Type State = ESlateCheckBoxState::Undetermined;
+	ECheckBoxState State = ECheckBoxState::Undetermined;
 
 	uint32 TotalNumChildren = ValidCollisionChannels.Num();
 	if (TotalNumChildren >= 1)
@@ -918,7 +918,7 @@ ESlateCheckBoxState::Type FBodyInstanceCustomization::IsAllCollisionChannelCheck
 		{
 			if (State != IsCollisionChannelChecked(Index, InCollisionResponse))
 			{
-				State = ESlateCheckBoxState::Undetermined;
+				State = ECheckBoxState::Undetermined;
 				break;
 			}
 		}
diff --git a/Engine/Source/Editor/DetailCustomizations/Private/BodyInstanceCustomization.h b/Engine/Source/Editor/DetailCustomizations/Private/BodyInstanceCustomization.h
--- a/Engine/Source/Editor/DetailCustomizations/Private/BodyInstanceCustomization.h
+++ b/Engine/Source/Editor/DetailCustomizations/Private/BodyInstanceCustomization.h
@@ -48,11 +48,11 @@ private:
 	EVisibility ShouldShowResetToDefaultResponse(int32 ValidIndex) const;
 
 	// collision channel check boxes
-	void OnCollisionChannelChanged(ESlateCheckBoxState::Type InNewValue, int32 ValidIndex, ECollisionResponse InCollisionResponse);
-	ESlateCheckBoxState::Type IsCollisionChannelChecked( int32 ValidIndex, ECollisionResponse InCollisionResponse) const;
+	void OnCollisionChannelChanged(ECheckBoxState InNewValue, int32 ValidIndex, ECollisionResponse InCollisionResponse);
+	ECheckBoxState IsCollisionChannelChecked( int32 ValidIndex, ECollisionResponse InCollisionResponse) const;
 	// all collision channel check boxes
-	void OnAllCollisionChannelChanged(ESlateCheckBoxState::Type InNewValue, ECollisionResponse InCollisionResponse);
-	ESlateCheckBoxState::Type IsAllCollisionChannelChecked( ECollisionResponse InCollisionResponse) const;
+	void OnAllCollisionChannelChanged(ECheckBoxState InNewValue, ECollisionResponse InCollisionResponse);
+	ECheckBoxState IsAllCollisionChannelChecked( ECollisionResponse InCollisionResponse) const;
 
 	// should show custom prop
 	bool ShouldEnableCustomCollisionSetup() const;
diff --git a/Engine/Source/Editor/DetailCustomizations/Private/CollisionProfileDetails.cpp b/Engine/Source/Editor/DetailCustomizations/Private/CollisionProfileDetails.cpp
--- a/Engine/Source/Editor/DetailCustomizations/Private/CollisionProfileDetails.cpp
+++ b/Engine/Source/Editor/DetailCustomizations/Private/CollisionProfileDetails.cpp
@@ -400,11 +400,11 @@ private:
 	void AddCollisionChannel(TArray<FCollisionChannelInfo>	ValidCollisionChannels, bool bTraceType);
 
 	// collision channel check boxes
-	void OnCollisionChannelChanged(ESlateCheckBoxState::Type InNewValue, int32 ValidIndex, ECollisionResponse InCollisionResponse);
-	ESlateCheckBoxState::Type IsCollisionChannelChecked(int32 ValidIndex, ECollisionResponse InCollisionResponse) const;
+	void OnCollisionChannelChanged(ECheckBoxState InNewValue, int32 ValidIndex, ECollisionResponse InCollisionResponse);
+	ECheckBoxState IsCollisionChannelChecked(int32 ValidIndex, ECollisionResponse InCollisionResponse) const;
 	// all collision channel check boxes
-	void OnAllCollisionChannelChanged(ESlateCheckBoxState::Type InNewValue, ECollisionResponse InCollisionResponse);
-	ESlateCheckBoxState::Type IsAllCollisionChannelChecked(ECollisionResponse InCollisionResponse) const;
+	void OnAllCollisionChannelChanged(ECheckBoxState InNewValue, ECollisionResponse InCollisionResponse);
+	ECheckBoxState IsAllCollisionChannelChecked(ECollisionResponse InCollisionResponse) const;
 };
 
 void SProfileEditDialog::Construct(const FArguments& InArgs)
@@ -1084,7 +1084,7 @@ void SProfileEditDialog::AddCollisionChannel(TArray<FCollisionChannelInfo>	Valid
 	}
 }
 
-void SProfileEditDialog::OnCollisionChannelChanged(ESlateCheckBoxState::Type InNewValue, int32 ValidIndex, ECollisionResponse InCollisionResponse)
+void SProfileEditDialog::OnCollisionChannelChanged(ECheckBoxState InNewValue, int32 ValidIndex, ECollisionResponse InCollisionResponse)
 {
 	if(ValidIndex >= 0 && ValidIndex < MAX_COLLISION_CHANNEL)
 	{
@@ -1092,7 +1092,7 @@ void SProfileEditDialog::OnCollisionChannelChanged(ESlateCheckBoxState::Type InN
 	}
 }
 
-ESlateCheckBoxState::Type SProfileEditDialog::IsCollisionChannelChecked(int32 ValidIndex, ECollisionResponse InCollisionResponse) const
+ECheckBoxState SProfileEditDialog::IsCollisionChannelChecked(int32 ValidIndex, ECollisionResponse InCollisionResponse) const
 {
 	TArray<uint8> CollisionResponses;
 
@@ -1100,14 +1100,14 @@ ESlateCheckBoxState::Type SProfileEditDialog::IsCollisionChannelChecked(int32 Va
 	{
 		if(ProfileTemplate.ResponseToChannels.EnumArray[ValidIndex] == InCollisionResponse)
 		{
-			return ESlateCheckBoxState::Checked;
+			return ECheckBoxState::Checked;
 		}
 	}
 
-	return ESlateCheckBoxState::Unchecked;
+	return ECheckBoxState::Unchecked;
 }
 
-void SProfileEditDialog::OnAllCollisionChannelChanged(ESlateCheckBoxState::Type InNewValue, ECollisionResponse InCollisionResponse)
+void SProfileEditDialog::OnAllCollisionChannelChanged(ECheckBoxState InNewValue, ECollisionResponse InCollisionResponse)
 {
 	for(int32 Index=0; Index<MAX_COLLISION_CHANNEL; ++Index)
 	{
@@ -1115,17 +1115,17 @@ void SProfileEditDialog::OnAllCollisionChannelChanged(ESlateCheckBoxState::Type
 	}
 }
 
-ESlateCheckBoxState::Type SProfileEditDialog::IsAllCollisionChannelChecked(ECollisionResponse InCollisionResponse) const
+ECheckBoxState SProfileEditDialog::IsAllCollisionChannelChecked(ECollisionResponse InCollisionResponse) const
 {
 	for(int32 Index=0; Index<MAX_COLLISION_CHANNEL; ++Index)
 	{
 		if(ProfileTemplate.ResponseToChannels.EnumArray[Index] != InCollisionResponse)
 		{
-			return ESlateCheckBoxState::Unchecked;
+			return ECheckBoxState::Unchecked;
 		}
 	}
 
-	return ESlateCheckBoxState::Checked;
+	return ECheckBoxState::Checked;
 }
 //====================================================================================
 // SChannelListItem 
diff --git a/Engine/Source/Editor/DetailCustomizations/Private/ComponentTransformDetails.cpp b/Engine/Source/Editor/DetailCustomizations/Private/ComponentTransformDetails.cpp
--- a/Engine/Source/Editor/DetailCustomizations/Private/ComponentTransformDetails.cpp
+++ b/Engine/Source/Editor/DetailCustomizations/Private/ComponentTransformDetails.cpp
@@ -569,14 +569,14 @@ const FSlateBrush* FComponentTransformDetails::GetPreserveScaleRatioImage() cons
 	return bPreserveScaleRatio ? FEditorStyle::GetBrush( TEXT("GenericLock") ) : FEditorStyle::GetBrush( TEXT("GenericUnlock") ) ;
 }
 
-ESlateCheckBoxState::Type FComponentTransformDetails::IsPreserveScaleRatioChecked() const
+ECheckBoxState FComponentTransformDetails::IsPreserveScaleRatioChecked() const
 {
-	return bPreserveScaleRatio ? ESlateCheckBoxState::Checked : ESlateCheckBoxState::Unchecked;
+	return bPreserveScaleRatio ? ECheckBoxState::Checked : ECheckBoxState::Unchecked;
 }
 
-void FComponentTransformDetails::OnPreserveScaleRatioToggled( ESlateCheckBoxState::Type NewState )
+void FComponentTransformDetails::OnPreserveScaleRatioToggled( ECheckBoxState NewState )
 {
-	bPreserveScaleRatio = (NewState == ESlateCheckBoxState::Checked) ? true : false;
+	bPreserveScaleRatio = (NewState == ECheckBoxState::Checked) ? true : false;
 	GConfig->SetBool(TEXT("SelectionDetails"), TEXT("PreserveScaleRatio"), bPreserveScaleRatio, GEditorUserSettingsIni);
 }
 
diff --git a/Engine/Source/Editor/DetailCustomizations/Private/ComponentTransformDetails.h b/Engine/Source/Editor/DetailCustomizations/Private/ComponentTransformDetails.h
--- a/Engine/Source/Editor/DetailCustomizations/Private/ComponentTransformDetails.h
+++ b/Engine/Source/Editor/DetailCustomizations/Private/ComponentTransformDetails.h
@@ -100,12 +100,12 @@ private:
 	const FSlateBrush* GetPreserveScaleRatioImage() const;
 
 	/** @return The state of the preserve scale ratio checkbox */
-	ESlateCheckBoxState::Type IsPreserveScaleRatioChecked() const;
+	ECheckBoxState IsPreserveScaleRatioChecked() const;
 
 	/**
 	 * Called when the preserve scale ratio checkbox is toggled
 	 */
-	void OnPreserveScaleRatioToggled( ESlateCheckBoxState::Type NewState );
+	void OnPreserveScaleRatioToggled( ECheckBoxState NewState );
 
 	/**
 	 * Builds a transform field label
diff --git a/Engine/Source/Editor/DetailCustomizations/Private/EnvQueryParamInstanceCustomization.cpp b/Engine/Source/Editor/DetailCustomizations/Private/EnvQueryParamInstanceCustomization.cpp
--- a/Engine/Source/Editor/DetailCustomizations/Private/EnvQueryParamInstanceCustomization.cpp
+++ b/Engine/Source/Editor/DetailCustomizations/Private/EnvQueryParamInstanceCustomization.cpp
@@ -110,7 +110,7 @@ void FEnvQueryParamInstanceCustomization::OnParamNumValueChanged(float FloatValu
 	}
 }
 
-ESlateCheckBoxState::Type FEnvQueryParamInstanceCustomization::GetParamBoolValue() const
+ECheckBoxState FEnvQueryParamInstanceCustomization::GetParamBoolValue() const
 {
 	if (ParamType == EEnvQueryParam::Bool)
 	{
@@ -118,20 +118,20 @@ ESlateCheckBoxState::Type FEnvQueryParamInstanceCustomization::GetParamBoolValue
 		FPropertyAccess::Result PropResult = ValueProp->GetValue(StoreValue);
 		if (PropResult == FPropertyAccess::Success)
 		{
-			return (StoreValue > 0.0f) ? ESlateCheckBoxState::Checked : ESlateCheckBoxState::Unchecked;
+			return (StoreValue > 0.0f) ? ECheckBoxState::Checked : ECheckBoxState::Unchecked;
 		}
 	}
 	
-	return ESlateCheckBoxState::Undetermined;
+	return ECheckBoxState::Undetermined;
 }
 
-void FEnvQueryParamInstanceCustomization::OnParamBoolValueChanged(ESlateCheckBoxState::Type BoolValue) const
+void FEnvQueryParamInstanceCustomization::OnParamBoolValueChanged(ECheckBoxState BoolValue) const
 {
 	if (ParamType == EEnvQueryParam::Bool)
 	{
-		const float StoreValue = (BoolValue == ESlateCheckBoxState::Checked) ? 1.0f : -1.0f;
+		const float StoreValue = (BoolValue == ECheckBoxState::Checked) ? 1.0f : -1.0f;
 		ValueProp->SetValue(StoreValue);
-		CachedBool = (BoolValue == ESlateCheckBoxState::Checked);
+		CachedBool = (BoolValue == ECheckBoxState::Checked);
 	}
 }
 
@@ -176,7 +176,7 @@ void FEnvQueryParamInstanceCustomization::InitCachedTypes()
 		{
 		case EEnvQueryParam::Float:	CachedFloat = GetParamNumValue().GetValue(); break;
 		case EEnvQueryParam::Int:	CachedInt = FMath::TruncToInt(GetParamNumValue().GetValue()); break;
-		case EEnvQueryParam::Bool:	CachedBool = (GetParamBoolValue() == ESlateCheckBoxState::Checked); break;
+		case EEnvQueryParam::Bool:	CachedBool = (GetParamBoolValue() == ECheckBoxState::Checked); break;
 		}
 	}
 }
@@ -192,7 +192,7 @@ void FEnvQueryParamInstanceCustomization::OnTypeChanged()
 		{
 		case EEnvQueryParam::Float:	OnParamNumValueChanged(CachedFloat); break;
 		case EEnvQueryParam::Int:	OnParamNumValueChanged(CachedInt); break;
-		case EEnvQueryParam::Bool:	OnParamBoolValueChanged(CachedBool ? ESlateCheckBoxState::Checked : ESlateCheckBoxState::Unchecked); break;
+		case EEnvQueryParam::Bool:	OnParamBoolValueChanged(CachedBool ? ECheckBoxState::Checked : ECheckBoxState::Unchecked); break;
 		}
 	}
 }
diff --git a/Engine/Source/Editor/DetailCustomizations/Private/EnvQueryParamInstanceCustomization.h b/Engine/Source/Editor/DetailCustomizations/Private/EnvQueryParamInstanceCustomization.h
--- a/Engine/Source/Editor/DetailCustomizations/Private/EnvQueryParamInstanceCustomization.h
+++ b/Engine/Source/Editor/DetailCustomizations/Private/EnvQueryParamInstanceCustomization.h
@@ -27,8 +27,8 @@ protected:
 	void OnParamNumValueChanged(float FloatValue) const;
 	EVisibility GetParamNumValueVisibility() const;
 
-	ESlateCheckBoxState::Type GetParamBoolValue() const;
-	void OnParamBoolValueChanged(ESlateCheckBoxState::Type BoolValue) const;
+	ECheckBoxState GetParamBoolValue() const;
+	void OnParamBoolValueChanged(ECheckBoxState BoolValue) const;
 	EVisibility GetParamBoolValueVisibility() const;
 
 	FString GetHeaderDesc() const;
diff --git a/Engine/Source/Editor/DetailCustomizations/Private/MathStructCustomizations.cpp b/Engine/Source/Editor/DetailCustomizations/Private/MathStructCustomizations.cpp
--- a/Engine/Source/Editor/DetailCustomizations/Private/MathStructCustomizations.cpp
+++ b/Engine/Source/Editor/DetailCustomizations/Private/MathStructCustomizations.cpp
@@ -103,14 +103,14 @@ const FSlateBrush* FMathStructCustomization::GetPreserveScaleRatioImage() const
 	return bPreserveScaleRatio ? FEditorStyle::GetBrush(TEXT("GenericLock")) : FEditorStyle::GetBrush(TEXT("GenericUnlock"));
 }
 
-ESlateCheckBoxState::Type FMathStructCustomization::IsPreserveScaleRatioChecked() const
+ECheckBoxState FMathStructCustomization::IsPreserveScaleRatioChecked() const
 {
-	return bPreserveScaleRatio ? ESlateCheckBoxState::Checked : ESlateCheckBoxState::Unchecked;
+	return bPreserveScaleRatio ? ECheckBoxState::Checked : ECheckBoxState::Unchecked;
 }
 
-void FMathStructCustomization::OnPreserveScaleRatioToggled(ESlateCheckBoxState::Type NewState, TWeakPtr<IPropertyHandle> PropertyHandle)
+void FMathStructCustomization::OnPreserveScaleRatioToggled(ECheckBoxState NewState, TWeakPtr<IPropertyHandle> PropertyHandle)
 {
-	bPreserveScaleRatio = ( NewState == ESlateCheckBoxState::Checked ) ? true : false;
+	bPreserveScaleRatio = ( NewState == ECheckBoxState::Checked ) ? true : false;
 
 	if ( PropertyHandle.IsValid() )
 	{
diff --git a/Engine/Source/Editor/DetailCustomizations/Private/MathStructCustomizations.h b/Engine/Source/Editor/DetailCustomizations/Private/MathStructCustomizations.h
--- a/Engine/Source/Editor/DetailCustomizations/Private/MathStructCustomizations.h
+++ b/Engine/Source/Editor/DetailCustomizations/Private/MathStructCustomizations.h
@@ -98,10 +98,10 @@ private:
 	const FSlateBrush* GetPreserveScaleRatioImage() const;
 
 	/** Gets the checked value of the preserve scale ratio option */
-	ESlateCheckBoxState::Type IsPreserveScaleRatioChecked() const;
+	ECheckBoxState IsPreserveScaleRatioChecked() const;
 
 	/** Called when the user toggles preserve ratio. */
-	void OnPreserveScaleRatioToggled(ESlateCheckBoxState::Type NewState, TWeakPtr<IPropertyHandle> PropertyHandle);
+	void OnPreserveScaleRatioToggled(ECheckBoxState NewState, TWeakPtr<IPropertyHandle> PropertyHandle);
 
 private:
 	/** Called when a value starts to be changed by a slider */
diff --git a/Engine/Source/Editor/DetailCustomizations/Private/PhysicsConstraintComponentDetails.cpp b/Engine/Source/Editor/DetailCustomizations/Private/PhysicsConstraintComponentDetails.cpp
--- a/Engine/Source/Editor/DetailCustomizations/Private/PhysicsConstraintComponentDetails.cpp
+++ b/Engine/Source/Editor/DetailCustomizations/Private/PhysicsConstraintComponentDetails.cpp
@@ -458,19 +458,19 @@ FReply FPhysicsConstraintComponentDetails::OnBallSocketClicked()
 	return FReply::Handled();
 }
 
-ESlateCheckBoxState::Type FPhysicsConstraintComponentDetails::IsLimitRadioChecked( TSharedPtr<IPropertyHandle> Property, uint8 Value ) const
+ECheckBoxState FPhysicsConstraintComponentDetails::IsLimitRadioChecked( TSharedPtr<IPropertyHandle> Property, uint8 Value ) const
 {
 	uint8 PropertyEnumValue = 0;
 	if (Property.IsValid() && Property->GetValue(PropertyEnumValue))
 	{
-		return PropertyEnumValue == Value ? ESlateCheckBoxState::Checked : ESlateCheckBoxState::Unchecked;
+		return PropertyEnumValue == Value ? ECheckBoxState::Checked : ECheckBoxState::Unchecked;
 	}
-	return ESlateCheckBoxState::Unchecked;
+	return ECheckBoxState::Unchecked;
 }
 
-void FPhysicsConstraintComponentDetails::OnLimitRadioChanged( ESlateCheckBoxState::Type CheckType, TSharedPtr<IPropertyHandle> Property, uint8 Value )
+void FPhysicsConstraintComponentDetails::OnLimitRadioChanged( ECheckBoxState CheckType, TSharedPtr<IPropertyHandle> Property, uint8 Value )
 {
-	if (Property.IsValid() && CheckType == ESlateCheckBoxState::Checked)
+	if (Property.IsValid() && CheckType == ECheckBoxState::Checked)
 	{
 		Property->SetValue(Value);
 	}
diff --git a/Engine/Source/Editor/DetailCustomizations/Private/PhysicsConstraintComponentDetails.h b/Engine/Source/Editor/DetailCustomizations/Private/PhysicsConstraintComponentDetails.h
--- a/Engine/Source/Editor/DetailCustomizations/Private/PhysicsConstraintComponentDetails.h
+++ b/Engine/Source/Editor/DetailCustomizations/Private/PhysicsConstraintComponentDetails.h
@@ -46,8 +46,8 @@ private:
 	FReply OnPrismaticClicked();
 	FReply OnBallSocketClicked();
 
-	ESlateCheckBoxState::Type IsLimitRadioChecked(TSharedPtr<IPropertyHandle> Property, uint8 Value) const ;
-	void OnLimitRadioChanged(ESlateCheckBoxState::Type CheckType, TSharedPtr<IPropertyHandle> Property, uint8 Value);
+	ECheckBoxState IsLimitRadioChecked(TSharedPtr<IPropertyHandle> Property, uint8 Value) const ;
+	void OnLimitRadioChanged(ECheckBoxState CheckType, TSharedPtr<IPropertyHandle> Property, uint8 Value);
 
 
 	TSharedPtr<IPropertyHandle> LinearPositionDriveProperty;
diff --git a/Engine/Source/Editor/DetailCustomizations/Private/PrimitiveComponentDetails.h b/Engine/Source/Editor/DetailCustomizations/Private/PrimitiveComponentDetails.h
--- a/Engine/Source/Editor/DetailCustomizations/Private/PrimitiveComponentDetails.h
+++ b/Engine/Source/Editor/DetailCustomizations/Private/PrimitiveComponentDetails.h
@@ -12,9 +12,9 @@ public:
 	virtual void CustomizeDetails( IDetailLayoutBuilder& DetailBuilder ) override;
 
 private:
-	ESlateCheckBoxState::Type IsMobilityActive(TSharedRef<IPropertyHandle> MobilityHandle, EComponentMobility::Type InMobility) const;
+	ECheckBoxState IsMobilityActive(TSharedRef<IPropertyHandle> MobilityHandle, EComponentMobility::Type InMobility) const;
 
-	void OnMobilityChanged(ESlateCheckBoxState::Type InCheckedState, TSharedRef<IPropertyHandle> MobilityHandle, EComponentMobility::Type InMobility);
+	void OnMobilityChanged(ECheckBoxState InCheckedState, TSharedRef<IPropertyHandle> MobilityHandle, EComponentMobility::Type InMobility);
 
 	void AddAdvancedSubCategory(IDetailLayoutBuilder& DetailBuilder, FName MainCategory, FName SubCategory);
 
diff --git a/Engine/Source/Editor/DetailCustomizations/Private/RigDetails.cpp b/Engine/Source/Editor/DetailCustomizations/Private/RigDetails.cpp
--- a/Engine/Source/Editor/DetailCustomizations/Private/RigDetails.cpp
+++ b/Engine/Source/Editor/DetailCustomizations/Private/RigDetails.cpp
@@ -506,21 +506,21 @@ void FRigDetails::OnComboBoxOopening(TSharedRef<IPropertyHandle> ParentSpaceProp
 	ComboBox->SetSelectedItem(SelectedItem);
 }
 
-void FRigDetails::OnAdvancedCheckBoxStateChanged(ESlateCheckBoxState::Type NewState, TSharedRef<IPropertyHandle> PropertyHandle)
+void FRigDetails::OnAdvancedCheckBoxStateChanged(ECheckBoxState NewState, TSharedRef<IPropertyHandle> PropertyHandle)
 {
-	bool bValue = (NewState == ESlateCheckBoxState::Checked)? true : false;
+	bool bValue = (NewState == ECheckBoxState::Checked)? true : false;
 	PropertyHandle->SetValue(bValue);
 }
 
-ESlateCheckBoxState::Type FRigDetails::AdvancedCheckBoxIsChecked(TSharedRef<IPropertyHandle> PropertyHandle) const
+ECheckBoxState FRigDetails::AdvancedCheckBoxIsChecked(TSharedRef<IPropertyHandle> PropertyHandle) const
 {
 	bool bValue = false;
 	// multi value doesn't work in array, so i'm not handling multi value
 	if (PropertyHandle->GetValue(bValue) != FPropertyAccess::Fail)
 	{
-		return (bValue)? ESlateCheckBoxState::Checked: ESlateCheckBoxState::Unchecked;
+		return (bValue)? ECheckBoxState::Checked: ECheckBoxState::Unchecked;
 	}
 
-	return ESlateCheckBoxState::Undetermined;
+	return ECheckBoxState::Undetermined;
 }
 #undef LOCTEXT_NAMESPACE
diff --git a/Engine/Source/Editor/DetailCustomizations/Private/RigDetails.h b/Engine/Source/Editor/DetailCustomizations/Private/RigDetails.h
--- a/Engine/Source/Editor/DetailCustomizations/Private/RigDetails.h
+++ b/Engine/Source/Editor/DetailCustomizations/Private/RigDetails.h
@@ -45,8 +45,8 @@ private:
 	void OnComboBoxOopening(TSharedRef<IPropertyHandle> ParentSpacePropertyHandle, int32 ArrayIndex, bool bTranslation);
 	// check box handler
 	// Callback for changing this row's Share check box state.
-	void OnAdvancedCheckBoxStateChanged(ESlateCheckBoxState::Type NewState, TSharedRef<IPropertyHandle> PropertyHandle);
-	ESlateCheckBoxState::Type AdvancedCheckBoxIsChecked(TSharedRef<IPropertyHandle> PropertyHandle) const;
+	void OnAdvancedCheckBoxStateChanged(ECheckBoxState NewState, TSharedRef<IPropertyHandle> PropertyHandle);
+	ECheckBoxState AdvancedCheckBoxIsChecked(TSharedRef<IPropertyHandle> PropertyHandle) const;
 
 	// button handlers
 	FReply OnSetAllToWorld();
diff --git a/Engine/Source/Editor/DetailCustomizations/Private/SceneCaptureDetails.cpp b/Engine/Source/Editor/DetailCustomizations/Private/SceneCaptureDetails.cpp
--- a/Engine/Source/Editor/DetailCustomizations/Private/SceneCaptureDetails.cpp
+++ b/Engine/Source/Editor/DetailCustomizations/Private/SceneCaptureDetails.cpp
@@ -193,7 +193,7 @@ void FSceneCaptureDetails::CustomizeDetails( IDetailLayoutBuilder& DetailLayout
 	}
 }
 
-ESlateCheckBoxState::Type FSceneCaptureDetails::OnGetDisplayCheckState(FString ShowFlagName) const
+ECheckBoxState FSceneCaptureDetails::OnGetDisplayCheckState(FString ShowFlagName) const
 {
 	bool IsChecked = false;
 	FEngineShowFlagsSetting* FlagSetting;
@@ -210,14 +210,14 @@ ESlateCheckBoxState::Type FSceneCaptureDetails::OnGetDisplayCheckState(FString S
 		}
 	}
 
-	return IsChecked ? ESlateCheckBoxState::Checked : ESlateCheckBoxState::Unchecked;
+	return IsChecked ? ECheckBoxState::Checked : ECheckBoxState::Unchecked;
 }
 
-void FSceneCaptureDetails::OnShowFlagCheckStateChanged(ESlateCheckBoxState::Type InNewRadioState, FString FlagName)
+void FSceneCaptureDetails::OnShowFlagCheckStateChanged(ECheckBoxState InNewRadioState, FString FlagName)
 {
 	SceneCaptureComponent->Modify();
 	FEngineShowFlagsSetting* FlagSetting;
-	bool bNewCheckState = (InNewRadioState == ESlateCheckBoxState::Checked);
+	bool bNewCheckState = (InNewRadioState == ECheckBoxState::Checked);
 
 	// If setting exists, update it
 	if (SceneCaptureComponent->GetSettingForShowFlag(FlagName, &FlagSetting))
diff --git a/Engine/Source/Editor/DetailCustomizations/Private/SceneCaptureDetails.h b/Engine/Source/Editor/DetailCustomizations/Private/SceneCaptureDetails.h
--- a/Engine/Source/Editor/DetailCustomizations/Private/SceneCaptureDetails.h
+++ b/Engine/Source/Editor/DetailCustomizations/Private/SceneCaptureDetails.h
@@ -23,8 +23,8 @@ private:
 	* @param The type of checkbox.
 	* @return The desired checkbox state.
 	*/
-	ESlateCheckBoxState::Type OnGetDisplayCheckState(FString ShowFlagName) const;
+	ECheckBoxState OnGetDisplayCheckState(FString ShowFlagName) const;
 
 	/** Show flag settings changed, so update the scene capture */
-	void OnShowFlagCheckStateChanged(ESlateCheckBoxState::Type InNewRadioState, FString FlagName);
+	void OnShowFlagCheckStateChanged(ECheckBoxState InNewRadioState, FString FlagName);
 };
diff --git a/Engine/Source/Editor/DetailCustomizations/Private/SceneComponentDetails.cpp b/Engine/Source/Editor/DetailCustomizations/Private/SceneComponentDetails.cpp
--- a/Engine/Source/Editor/DetailCustomizations/Private/SceneComponentDetails.cpp
+++ b/Engine/Source/Editor/DetailCustomizations/Private/SceneComponentDetails.cpp
@@ -350,17 +350,17 @@ TSharedRef<IDetailCustomization> FSceneComponentDetails::MakeInstance()
 	return MakeShareable( new FSceneComponentDetails );
 }
 
-ESlateCheckBoxState::Type FSceneComponentDetails::IsMobilityActive(TWeakPtr<IPropertyHandle> MobilityHandle, EComponentMobility::Type InMobility) const
+ECheckBoxState FSceneComponentDetails::IsMobilityActive(TWeakPtr<IPropertyHandle> MobilityHandle, EComponentMobility::Type InMobility) const
 {
 	if ( MobilityHandle.IsValid() )
 	{
 		uint8 MobilityByte;
 		MobilityHandle.Pin()->GetValue(MobilityByte);
 
-		return MobilityByte == InMobility ? ESlateCheckBoxState::Checked : ESlateCheckBoxState::Unchecked;
+		return MobilityByte == InMobility ? ECheckBoxState::Checked : ECheckBoxState::Unchecked;
 	}
 
-	return ESlateCheckBoxState::Unchecked;
+	return ECheckBoxState::Unchecked;
 }
 
 FSlateColor FSceneComponentDetails::GetMobilityTextColor(TWeakPtr<IPropertyHandle> MobilityHandle, EComponentMobility::Type InMobility) const
@@ -376,9 +376,9 @@ FSlateColor FSceneComponentDetails::GetMobilityTextColor(TWeakPtr<IPropertyHandl
 	return FSlateColor(FLinearColor(0.72f, 0.72f, 0.72f, 1.f));
 }
 
-void FSceneComponentDetails::OnMobilityChanged(ESlateCheckBoxState::Type InCheckedState, TWeakPtr<IPropertyHandle> MobilityHandle, EComponentMobility::Type InMobility)
+void FSceneComponentDetails::OnMobilityChanged(ECheckBoxState InCheckedState, TWeakPtr<IPropertyHandle> MobilityHandle, EComponentMobility::Type InMobility)
 {
-	if ( MobilityHandle.IsValid() && InCheckedState == ESlateCheckBoxState::Checked)
+	if ( MobilityHandle.IsValid() && InCheckedState == ECheckBoxState::Checked)
 	{
 		MobilityHandle.Pin()->SetValue((uint8)InMobility);
 	}
diff --git a/Engine/Source/Editor/DetailCustomizations/Private/SceneComponentDetails.h b/Engine/Source/Editor/DetailCustomizations/Private/SceneComponentDetails.h
--- a/Engine/Source/Editor/DetailCustomizations/Private/SceneComponentDetails.h
+++ b/Engine/Source/Editor/DetailCustomizations/Private/SceneComponentDetails.h
@@ -16,9 +16,9 @@ private:
 
 	FSlateColor GetMobilityTextColor(TWeakPtr<IPropertyHandle> MobilityHandle, EComponentMobility::Type InMobility) const;
 
-	ESlateCheckBoxState::Type IsMobilityActive(TWeakPtr<IPropertyHandle> MobilityHandle, EComponentMobility::Type InMobility) const;
+	ECheckBoxState IsMobilityActive(TWeakPtr<IPropertyHandle> MobilityHandle, EComponentMobility::Type InMobility) const;
 
-	void OnMobilityChanged(ESlateCheckBoxState::Type InCheckedState, TWeakPtr<IPropertyHandle> MobilityHandle, EComponentMobility::Type InMobility);
+	void OnMobilityChanged(ECheckBoxState InCheckedState, TWeakPtr<IPropertyHandle> MobilityHandle, EComponentMobility::Type InMobility);
 
 	FText GetMobilityToolTip(TWeakPtr<IPropertyHandle> MobilityHandle) const;
 	
diff --git a/Engine/Source/Editor/DetailCustomizations/Private/SlateColorCustomization.cpp b/Engine/Source/Editor/DetailCustomizations/Private/SlateColorCustomization.cpp
--- a/Engine/Source/Editor/DetailCustomizations/Private/SlateColorCustomization.cpp
+++ b/Engine/Source/Editor/DetailCustomizations/Private/SlateColorCustomization.cpp
@@ -73,22 +73,22 @@ void FSlateColorCustomization::OnValueChanged()
 	ColorRuleHandle->SetValueFromFormattedString(TEXT("UseColor_Specified"));
 }
 
-ESlateCheckBoxState::Type FSlateColorCustomization::GetForegroundCheckState() const
+ECheckBoxState FSlateColorCustomization::GetForegroundCheckState() const
 {
 	FString ColorRuleValue;
 	ColorRuleHandle->GetValueAsFormattedString(ColorRuleValue);
 
 	if ( ColorRuleValue == TEXT("UseColor_Foreground") )
 	{
-		return ESlateCheckBoxState::Checked;
+		return ECheckBoxState::Checked;
 	}
 
-	return ESlateCheckBoxState::Unchecked;
+	return ECheckBoxState::Unchecked;
 }
 
-void FSlateColorCustomization::HandleForegroundChanged(ESlateCheckBoxState::Type CheckedState)
+void FSlateColorCustomization::HandleForegroundChanged(ECheckBoxState CheckedState)
 {
-	if ( CheckedState == ESlateCheckBoxState::Checked )
+	if ( CheckedState == ECheckBoxState::Checked )
 	{
 		ColorRuleHandle->SetValueFromFormattedString(TEXT("UseColor_Foreground"));
 	}
diff --git a/Engine/Source/Editor/DetailCustomizations/Private/SlateColorCustomization.h b/Engine/Source/Editor/DetailCustomizations/Private/SlateColorCustomization.h
--- a/Engine/Source/Editor/DetailCustomizations/Private/SlateColorCustomization.h
+++ b/Engine/Source/Editor/DetailCustomizations/Private/SlateColorCustomization.h
@@ -21,9 +21,9 @@ private:
 	 */
 	virtual void OnValueChanged();
 
-	ESlateCheckBoxState::Type GetForegroundCheckState() const;
+	ECheckBoxState GetForegroundCheckState() const;
 
-	void HandleForegroundChanged(ESlateCheckBoxState::Type CheckedState);
+	void HandleForegroundChanged(ECheckBoxState CheckedState);
 
 private:
 
diff --git a/Engine/Source/Editor/DetailCustomizations/Private/WindowsTargetSettingsDetails.cpp b/Engine/Source/Editor/DetailCustomizations/Private/WindowsTargetSettingsDetails.cpp
--- a/Engine/Source/Editor/DetailCustomizations/Private/WindowsTargetSettingsDetails.cpp
+++ b/Engine/Source/Editor/DetailCustomizations/Private/WindowsTargetSettingsDetails.cpp
@@ -286,7 +286,7 @@ void FTargetShaderFormatsPropertyDetails::CreateTargetShaderFormatsPropertyView(
 }
 
 
-void FTargetShaderFormatsPropertyDetails::OnTargetedRHIChanged(ESlateCheckBoxState::Type InNewValue, FName InRHIName)
+void FTargetShaderFormatsPropertyDetails::OnTargetedRHIChanged(ECheckBoxState InNewValue, FName InRHIName)
 {
 	TArray<void*> RawPtrs;
 	TargetShaderFormatsPropertyHandle->AccessRawData(RawPtrs);
@@ -297,7 +297,7 @@ void FTargetShaderFormatsPropertyDetails::OnTargetedRHIChanged(ESlateCheckBoxSta
 		for (void* RawPtr : RawPtrs)
 		{
 			TArray<FString>& Array = *(TArray<FString>*)RawPtr;
-			if(InNewValue == ESlateCheckBoxState::Checked)
+			if(InNewValue == ECheckBoxState::Checked)
 			{
 				Array.Add(InRHIName.ToString());
 			}
@@ -311,9 +311,9 @@ void FTargetShaderFormatsPropertyDetails::OnTargetedRHIChanged(ESlateCheckBoxSta
 }
 
 
-ESlateCheckBoxState::Type FTargetShaderFormatsPropertyDetails::IsTargetedRHIChecked(FName InRHIName) const
+ECheckBoxState FTargetShaderFormatsPropertyDetails::IsTargetedRHIChecked(FName InRHIName) const
 {
-	ESlateCheckBoxState::Type CheckState = ESlateCheckBoxState::Unchecked;
+	ECheckBoxState CheckState = ECheckBoxState::Unchecked;
 
 	TArray<void*> RawPtrs;
 	TargetShaderFormatsPropertyHandle->AccessRawData(RawPtrs);
@@ -323,7 +323,7 @@ ESlateCheckBoxState::Type FTargetShaderFormatsPropertyDetails::IsTargetedRHIChec
 		TArray<FString>& Array = *(TArray<FString>*)RawPtr;
 		if(Array.Contains(InRHIName.ToString()))
 		{
-			CheckState = ESlateCheckBoxState::Checked;
+			CheckState = ECheckBoxState::Checked;
 		}
 	}
 	return CheckState;
diff --git a/Engine/Source/Editor/DetailCustomizations/Private/WindowsTargetSettingsDetails.h b/Engine/Source/Editor/DetailCustomizations/Private/WindowsTargetSettingsDetails.h
--- a/Engine/Source/Editor/DetailCustomizations/Private/WindowsTargetSettingsDetails.h
+++ b/Engine/Source/Editor/DetailCustomizations/Private/WindowsTargetSettingsDetails.h
@@ -31,8 +31,8 @@ public:
 private:
 
 	// Supported/Targeted RHI check boxes
-	void OnTargetedRHIChanged(ESlateCheckBoxState::Type InNewValue, FName InRHIName);
-	ESlateCheckBoxState::Type IsTargetedRHIChecked(FName InRHIName) const;
+	void OnTargetedRHIChanged(ECheckBoxState InNewValue, FName InRHIName);
+	ECheckBoxState IsTargetedRHIChecked(FName InRHIName) const;
 
 private:
 
diff --git a/Engine/Source/Editor/EnvironmentQueryEditor/Private/DetailCustomizations/EnvQueryParamSetupCustomization.cpp b/Engine/Source/Editor/EnvironmentQueryEditor/Private/DetailCustomizations/EnvQueryParamSetupCustomization.cpp
--- a/Engine/Source/Editor/EnvironmentQueryEditor/Private/DetailCustomizations/EnvQueryParamSetupCustomization.cpp
+++ b/Engine/Source/Editor/EnvironmentQueryEditor/Private/DetailCustomizations/EnvQueryParamSetupCustomization.cpp
@@ -154,7 +154,7 @@ void FEnvQueryParamSetupCustomization::OnParamNameCommitted(const FText& ParamNa
 	LastPropName = (NameValue != NAME_None) ? NameValue : FName(TEXT("unknown"));
 }
 
-ESlateCheckBoxState::Type FEnvQueryParamSetupCustomization::GetParamBoolValue() const
+ECheckBoxState FEnvQueryParamSetupCustomization::GetParamBoolValue() const
 {
 	if (ParamType == EEnvQueryParam::Bool)
 	{
@@ -162,18 +162,18 @@ ESlateCheckBoxState::Type FEnvQueryParamSetupCustomization::GetParamBoolValue()
 		FPropertyAccess::Result PropResult = ValueProp->GetValue(BoolValue);
 		if (PropResult == FPropertyAccess::Success)
 		{
-			return BoolValue ? ESlateCheckBoxState::Checked : ESlateCheckBoxState::Unchecked;
+			return BoolValue ? ECheckBoxState::Checked : ECheckBoxState::Unchecked;
 		}
 	}
 	
-	return ESlateCheckBoxState::Undetermined;
+	return ECheckBoxState::Undetermined;
 }
 
-void FEnvQueryParamSetupCustomization::OnParamBoolValueChanged(ESlateCheckBoxState::Type BoolValue) const
+void FEnvQueryParamSetupCustomization::OnParamBoolValueChanged(ECheckBoxState BoolValue) const
 {
 	if (ParamType == EEnvQueryParam::Bool)
 	{
-		const bool StoreValue = (BoolValue == ESlateCheckBoxState::Checked);
+		const bool StoreValue = (BoolValue == ECheckBoxState::Checked);
 		ValueProp->SetValue(StoreValue);
 	}
 }
diff --git a/Engine/Source/Editor/EnvironmentQueryEditor/Private/DetailCustomizations/EnvQueryParamSetupCustomization.h b/Engine/Source/Editor/EnvironmentQueryEditor/Private/DetailCustomizations/EnvQueryParamSetupCustomization.h
--- a/Engine/Source/Editor/EnvironmentQueryEditor/Private/DetailCustomizations/EnvQueryParamSetupCustomization.h
+++ b/Engine/Source/Editor/EnvironmentQueryEditor/Private/DetailCustomizations/EnvQueryParamSetupCustomization.h
@@ -33,8 +33,8 @@ protected:
 	void OnParamNumValueChanged(float FloatValue) const;
 	EVisibility GetParamNumValueVisibility() const;
 
-	ESlateCheckBoxState::Type GetParamBoolValue() const;
-	void OnParamBoolValueChanged(ESlateCheckBoxState::Type BoolValue) const;
+	ECheckBoxState GetParamBoolValue() const;
+	void OnParamBoolValueChanged(ECheckBoxState BoolValue) const;
 	EVisibility GetParamBoolValueVisibility() const;
 
 	void OnParamNameCommitted(const FText& ParamName, ETextCommit::Type CommitInfo);
diff --git a/Engine/Source/Editor/EnvironmentQueryEditor/Private/SGraphNode_EnvironmentQuery.cpp b/Engine/Source/Editor/EnvironmentQueryEditor/Private/SGraphNode_EnvironmentQuery.cpp
--- a/Engine/Source/Editor/EnvironmentQueryEditor/Private/SGraphNode_EnvironmentQuery.cpp
+++ b/Engine/Source/Editor/EnvironmentQueryEditor/Private/SGraphNode_EnvironmentQuery.cpp
@@ -515,18 +515,18 @@ EVisibility SGraphNode_EnvironmentQuery::GetTestToggleVisibility() const
 	return MyTestNode ? EVisibility::Visible : EVisibility::Collapsed;
 }
 
-ESlateCheckBoxState::Type SGraphNode_EnvironmentQuery::IsTestToggleChecked() const
+ECheckBoxState SGraphNode_EnvironmentQuery::IsTestToggleChecked() const
 {
 	UEnvironmentQueryGraphNode_Test* MyTestNode = Cast<UEnvironmentQueryGraphNode_Test>(GraphNode);
-	return MyTestNode && MyTestNode->bTestEnabled ? ESlateCheckBoxState::Checked : ESlateCheckBoxState::Unchecked;
+	return MyTestNode && MyTestNode->bTestEnabled ? ECheckBoxState::Checked : ECheckBoxState::Unchecked;
 }
 
-void SGraphNode_EnvironmentQuery::OnTestToggleChanged(ESlateCheckBoxState::Type NewState)
+void SGraphNode_EnvironmentQuery::OnTestToggleChanged(ECheckBoxState NewState)
 {
 	UEnvironmentQueryGraphNode_Test* MyTestNode = Cast<UEnvironmentQueryGraphNode_Test>(GraphNode);
 	if (MyTestNode)
 	{
-		MyTestNode->bTestEnabled = (NewState == ESlateCheckBoxState::Checked);
+		MyTestNode->bTestEnabled = (NewState == ECheckBoxState::Checked);
 		
 		if (MyTestNode->ParentNode)
 		{
diff --git a/Engine/Source/Editor/EnvironmentQueryEditor/Private/SGraphNode_EnvironmentQuery.h b/Engine/Source/Editor/EnvironmentQueryEditor/Private/SGraphNode_EnvironmentQuery.h
--- a/Engine/Source/Editor/EnvironmentQueryEditor/Private/SGraphNode_EnvironmentQuery.h
+++ b/Engine/Source/Editor/EnvironmentQueryEditor/Private/SGraphNode_EnvironmentQuery.h
@@ -41,8 +41,8 @@ public:
 	FSlateColor GetWeightProgressBarColor() const; 
 
 	EVisibility GetTestToggleVisibility() const;
-	ESlateCheckBoxState::Type IsTestToggleChecked() const;
-	void OnTestToggleChanged(ESlateCheckBoxState::Type NewState);
+	ECheckBoxState IsTestToggleChecked() const;
+	void OnTestToggleChanged(ECheckBoxState NewState);
 
 	/** gets drag over marker visibility */
 	EVisibility GetDragOverMarkerVisibility() const;
diff --git a/Engine/Source/Editor/FoliageEdit/Private/SFoliageEdit.cpp b/Engine/Source/Editor/FoliageEdit/Private/SFoliageEdit.cpp
--- a/Engine/Source/Editor/FoliageEdit/Private/SFoliageEdit.cpp
+++ b/Engine/Source/Editor/FoliageEdit/Private/SFoliageEdit.cpp
@@ -669,44 +669,44 @@ FReply SFoliageEdit::OnDrop_ListView(const FGeometry& MyGeometry, const FDragDro
 	return FReply::Handled();
 }
 
-void SFoliageEdit::OnCheckStateChanged_Landscape(ESlateCheckBoxState::Type InState)
+void SFoliageEdit::OnCheckStateChanged_Landscape(ECheckBoxState InState)
 {
-	FoliageEditMode->UISettings.SetFilterLandscape(InState == ESlateCheckBoxState::Checked ? true : false);
+	FoliageEditMode->UISettings.SetFilterLandscape(InState == ECheckBoxState::Checked ? true : false);
 }
 
-ESlateCheckBoxState::Type SFoliageEdit::GetCheckState_Landscape() const
+ECheckBoxState SFoliageEdit::GetCheckState_Landscape() const
 {
-	return FoliageEditMode->UISettings.GetFilterLandscape() ? ESlateCheckBoxState::Checked : ESlateCheckBoxState::Unchecked;
+	return FoliageEditMode->UISettings.GetFilterLandscape() ? ECheckBoxState::Checked : ECheckBoxState::Unchecked;
 }
 
-void SFoliageEdit::OnCheckStateChanged_StaticMesh(ESlateCheckBoxState::Type InState)
+void SFoliageEdit::OnCheckStateChanged_StaticMesh(ECheckBoxState InState)
 {
-	FoliageEditMode->UISettings.SetFilterStaticMesh(InState == ESlateCheckBoxState::Checked ? true : false);
+	FoliageEditMode->UISettings.SetFilterStaticMesh(InState == ECheckBoxState::Checked ? true : false);
 }
 
-ESlateCheckBoxState::Type SFoliageEdit::GetCheckState_StaticMesh() const
+ECheckBoxState SFoliageEdit::GetCheckState_StaticMesh() const
 {
-	return FoliageEditMode->UISettings.GetFilterStaticMesh() ? ESlateCheckBoxState::Checked : ESlateCheckBoxState::Unchecked;
+	return FoliageEditMode->UISettings.GetFilterStaticMesh() ? ECheckBoxState::Checked : ECheckBoxState::Unchecked;
 }
 
-void SFoliageEdit::OnCheckStateChanged_BSP(ESlateCheckBoxState::Type InState)
+void SFoliageEdit::OnCheckStateChanged_BSP(ECheckBoxState InState)
 {
-	FoliageEditMode->UISettings.SetFilterBSP(InState == ESlateCheckBoxState::Checked ? true : false);
+	FoliageEditMode->UISettings.SetFilterBSP(InState == ECheckBoxState::Checked ? true : false);
 }
 
-ESlateCheckBoxState::Type SFoliageEdit::GetCheckState_BSP() const
+ECheckBoxState SFoliageEdit::GetCheckState_BSP() const
 {
-	return FoliageEditMode->UISettings.GetFilterBSP() ? ESlateCheckBoxState::Checked : ESlateCheckBoxState::Unchecked;
+	return FoliageEditMode->UISettings.GetFilterBSP() ? ECheckBoxState::Checked : ECheckBoxState::Unchecked;
 }
 
-void SFoliageEdit::OnCheckStateChanged_Translucent(ESlateCheckBoxState::Type InState)
+void SFoliageEdit::OnCheckStateChanged_Translucent(ECheckBoxState InState)
 {
-	FoliageEditMode->UISettings.SetFilterTranslucent(InState == ESlateCheckBoxState::Checked ? true : false);
+	FoliageEditMode->UISettings.SetFilterTranslucent(InState == ECheckBoxState::Checked ? true : false);
 }
 
-ESlateCheckBoxState::Type SFoliageEdit::GetCheckState_Translucent() const
+ECheckBoxState SFoliageEdit::GetCheckState_Translucent() const
 {
-	return FoliageEditMode->UISettings.GetFilterTranslucent() ? ESlateCheckBoxState::Checked : ESlateCheckBoxState::Unchecked;
+	return FoliageEditMode->UISettings.GetFilterTranslucent() ? ECheckBoxState::Checked : ECheckBoxState::Unchecked;
 }
 
 EVisibility SFoliageEdit::GetVisibility_Radius() const
diff --git a/Engine/Source/Editor/FoliageEdit/Private/SFoliageEdit.h b/Engine/Source/Editor/FoliageEdit/Private/SFoliageEdit.h
--- a/Engine/Source/Editor/FoliageEdit/Private/SFoliageEdit.h
+++ b/Engine/Source/Editor/FoliageEdit/Private/SFoliageEdit.h
@@ -109,28 +109,28 @@ private:
 	TSharedRef<ITableRow> MakeWidgetFromOption(TSharedPtr<FFoliageMeshUIInfo> InItem, const TSharedRef<STableViewBase>& OwnerTable);
 
 	/** Sets the filter settings for if painting will occur on Landscapes. */
-	void OnCheckStateChanged_Landscape(ESlateCheckBoxState::Type InState);
+	void OnCheckStateChanged_Landscape(ECheckBoxState InState);
 
 	/** Retrieves the filter settings for painting on Landscapes. */
-	ESlateCheckBoxState::Type GetCheckState_Landscape() const;
+	ECheckBoxState GetCheckState_Landscape() const;
 
 	/** Sets the filter settings for if painting will occur on Static Meshes. */
-	void OnCheckStateChanged_StaticMesh(ESlateCheckBoxState::Type InState);
+	void OnCheckStateChanged_StaticMesh(ECheckBoxState InState);
 
 	/** Retrieves the filter settings for painting on Static Meshes. */
-	ESlateCheckBoxState::Type GetCheckState_StaticMesh() const;
+	ECheckBoxState GetCheckState_StaticMesh() const;
 
 	/** Sets the filter settings for if painting will occur on BSPs. */
-	void OnCheckStateChanged_BSP(ESlateCheckBoxState::Type InState);
+	void OnCheckStateChanged_BSP(ECheckBoxState InState);
 
 	/** Retrieves the filter settings for painting on BSPs. */
-	ESlateCheckBoxState::Type GetCheckState_BSP() const;
+	ECheckBoxState GetCheckState_BSP() const;
 
 	/** Sets the filter settings for if painting will occur on translucent meshes. */
-	void OnCheckStateChanged_Translucent(ESlateCheckBoxState::Type InState);
+	void OnCheckStateChanged_Translucent(ECheckBoxState InState);
 
 	/** Retrieves the filter settings for painting on translucent meshes. */
-	ESlateCheckBoxState::Type GetCheckState_Translucent() const;
+	ECheckBoxState GetCheckState_Translucent() const;
 
 	/** Checks if the text in the empty list overlay should appear. If the list is has items but the the drag and drop override is true, it will return EVisibility::Visible. */
 	EVisibility GetVisibility_EmptyList() const;
diff --git a/Engine/Source/Editor/FoliageEdit/Private/SFoliageEditMeshDisplayItem.cpp b/Engine/Source/Editor/FoliageEdit/Private/SFoliageEditMeshDisplayItem.cpp
--- a/Engine/Source/Editor/FoliageEdit/Private/SFoliageEditMeshDisplayItem.cpp
+++ b/Engine/Source/Editor/FoliageEdit/Private/SFoliageEditMeshDisplayItem.cpp
@@ -1830,14 +1830,14 @@ FSlateColor SFoliageEditMeshDisplayItem::GetBorderColor() const
 	return FSlateColor(FLinearColor(0.25f, 0.25f, 0.25f));
 }
 
-ESlateCheckBoxState::Type SFoliageEditMeshDisplayItem::IsSelected() const
+ECheckBoxState SFoliageEditMeshDisplayItem::IsSelected() const
 {
-	return FoliageSettingsPtr->IsSelected ? ESlateCheckBoxState::Checked : ESlateCheckBoxState::Unchecked;
+	return FoliageSettingsPtr->IsSelected ? ECheckBoxState::Checked : ECheckBoxState::Unchecked;
 }
 
-void SFoliageEditMeshDisplayItem::OnSelectionChanged(ESlateCheckBoxState::Type InType)
+void SFoliageEditMeshDisplayItem::OnSelectionChanged(ECheckBoxState InType)
 {
-	FoliageSettingsPtr->IsSelected = InType == ESlateCheckBoxState::Checked;
+	FoliageSettingsPtr->IsSelected = InType == ECheckBoxState::Checked;
 }
 
 FReply SFoliageEditMeshDisplayItem::OnSync()
@@ -1976,14 +1976,14 @@ EVisibility SFoliageEditMeshDisplayItem::IsNonUniformScalingVisible_Dummy() cons
 	return FoliageSettingsPtr->UniformScale ? EVisibility::Collapsed : EVisibility::Hidden;
 }
 
-void SFoliageEditMeshDisplayItem::OnDensityReapply(ESlateCheckBoxState::Type InState)
+void SFoliageEditMeshDisplayItem::OnDensityReapply(ECheckBoxState InState)
 {
-	FoliageSettingsPtr->ReapplyDensity = InState == ESlateCheckBoxState::Checked ? true : false;
+	FoliageSettingsPtr->ReapplyDensity = InState == ECheckBoxState::Checked ? true : false;
 }
 
-ESlateCheckBoxState::Type SFoliageEditMeshDisplayItem::IsDensityReapplyChecked() const
+ECheckBoxState SFoliageEditMeshDisplayItem::IsDensityReapplyChecked() const
 {
-	return FoliageSettingsPtr->ReapplyDensity ? ESlateCheckBoxState::Checked : ESlateCheckBoxState::Unchecked;
+	return FoliageSettingsPtr->ReapplyDensity ? ECheckBoxState::Checked : ECheckBoxState::Unchecked;
 }
 
 void SFoliageEditMeshDisplayItem::OnDensityChanged(float InValue)
@@ -2006,14 +2006,14 @@ float SFoliageEditMeshDisplayItem::GetDensityReapply() const
 	return FoliageSettingsPtr->ReapplyDensityAmount;
 }
 
-void SFoliageEditMeshDisplayItem::OnRadiusReapply(ESlateCheckBoxState::Type InState)
+void SFoliageEditMeshDisplayItem::OnRadiusReapply(ECheckBoxState InState)
 {
-	FoliageSettingsPtr->ReapplyRadius = InState == ESlateCheckBoxState::Checked ? true : false;
+	FoliageSettingsPtr->ReapplyRadius = InState == ECheckBoxState::Checked ? true : false;
 }
 
-ESlateCheckBoxState::Type SFoliageEditMeshDisplayItem::IsRadiusReapplyChecked() const
+ECheckBoxState SFoliageEditMeshDisplayItem::IsRadiusReapplyChecked() const
 {
-	return FoliageSettingsPtr->ReapplyRadius ? ESlateCheckBoxState::Checked : ESlateCheckBoxState::Unchecked;
+	return FoliageSettingsPtr->ReapplyRadius ? ECheckBoxState::Checked : ECheckBoxState::Unchecked;
 }
 
 void SFoliageEditMeshDisplayItem::OnRadiusChanged(float InValue)
@@ -2026,19 +2026,19 @@ float SFoliageEditMeshDisplayItem::GetRadius() const
 	return FoliageSettingsPtr->Radius;
 }
 
-void SFoliageEditMeshDisplayItem::OnAlignToNormalReapply(ESlateCheckBoxState::Type InState)
+void SFoliageEditMeshDisplayItem::OnAlignToNormalReapply(ECheckBoxState InState)
 {
-	FoliageSettingsPtr->ReapplyAlignToNormal = InState == ESlateCheckBoxState::Checked ? true : false;
+	FoliageSettingsPtr->ReapplyAlignToNormal = InState == ECheckBoxState::Checked ? true : false;
 }
 
-ESlateCheckBoxState::Type SFoliageEditMeshDisplayItem::IsAlignToNormalReapplyChecked() const
+ECheckBoxState SFoliageEditMeshDisplayItem::IsAlignToNormalReapplyChecked() const
 {
-	return FoliageSettingsPtr->ReapplyAlignToNormal ? ESlateCheckBoxState::Checked : ESlateCheckBoxState::Unchecked;
+	return FoliageSettingsPtr->ReapplyAlignToNormal ? ECheckBoxState::Checked : ECheckBoxState::Unchecked;
 }
 
-void SFoliageEditMeshDisplayItem::OnAlignToNormal(ESlateCheckBoxState::Type InState)
+void SFoliageEditMeshDisplayItem::OnAlignToNormal(ECheckBoxState InState)
 {
-	if (InState == ESlateCheckBoxState::Checked)
+	if (InState == ECheckBoxState::Checked)
 	{
 		FoliageSettingsPtr->AlignToNormal = true;
 	}
@@ -2048,9 +2048,9 @@ void SFoliageEditMeshDisplayItem::OnAlignToNormal(ESlateCheckBoxState::Type InSt
 	}
 }
 
-ESlateCheckBoxState::Type SFoliageEditMeshDisplayItem::IsAlignToNormalChecked() const
+ECheckBoxState SFoliageEditMeshDisplayItem::IsAlignToNormalChecked() const
 {
-	return FoliageSettingsPtr->AlignToNormal ? ESlateCheckBoxState::Checked : ESlateCheckBoxState::Unchecked;
+	return FoliageSettingsPtr->AlignToNormal ? ECheckBoxState::Checked : ECheckBoxState::Unchecked;
 }
 
 EVisibility SFoliageEditMeshDisplayItem::IsAlignToNormalVisible() const
@@ -2068,19 +2068,19 @@ float SFoliageEditMeshDisplayItem::GetMaxAngle() const
 	return FoliageSettingsPtr->AlignMaxAngle;
 }
 
-void SFoliageEditMeshDisplayItem::OnRandomYawReapply(ESlateCheckBoxState::Type InState)
+void SFoliageEditMeshDisplayItem::OnRandomYawReapply(ECheckBoxState InState)
 {
-	FoliageSettingsPtr->ReapplyRandomYaw = InState == ESlateCheckBoxState::Checked ? true : false;
+	FoliageSettingsPtr->ReapplyRandomYaw = InState == ECheckBoxState::Checked ? true : false;
 }
 
-ESlateCheckBoxState::Type SFoliageEditMeshDisplayItem::IsRandomYawReapplyChecked() const
+ECheckBoxState SFoliageEditMeshDisplayItem::IsRandomYawReapplyChecked() const
 {
-	return FoliageSettingsPtr->ReapplyRandomYaw ? ESlateCheckBoxState::Checked : ESlateCheckBoxState::Unchecked;
+	return FoliageSettingsPtr->ReapplyRandomYaw ? ECheckBoxState::Checked : ECheckBoxState::Unchecked;
 }
 
-void SFoliageEditMeshDisplayItem::OnRandomYaw(ESlateCheckBoxState::Type InState)
+void SFoliageEditMeshDisplayItem::OnRandomYaw(ECheckBoxState InState)
 {
-	if (InState == ESlateCheckBoxState::Checked)
+	if (InState == ECheckBoxState::Checked)
 	{
 		FoliageSettingsPtr->RandomYaw = true;
 	}
@@ -2090,14 +2090,14 @@ void SFoliageEditMeshDisplayItem::OnRandomYaw(ESlateCheckBoxState::Type InState)
 	}
 }
 
-ESlateCheckBoxState::Type SFoliageEditMeshDisplayItem::IsRandomYawChecked() const
+ECheckBoxState SFoliageEditMeshDisplayItem::IsRandomYawChecked() const
 {
-	return FoliageSettingsPtr->RandomYaw ? ESlateCheckBoxState::Checked : ESlateCheckBoxState::Unchecked;
+	return FoliageSettingsPtr->RandomYaw ? ECheckBoxState::Checked : ECheckBoxState::Unchecked;
 }
 
-void SFoliageEditMeshDisplayItem::OnUniformScale(ESlateCheckBoxState::Type InState)
+void SFoliageEditMeshDisplayItem::OnUniformScale(ECheckBoxState InState)
 {
-	if (InState == ESlateCheckBoxState::Checked)
+	if (InState == ECheckBoxState::Checked)
 	{
 		FoliageSettingsPtr->UniformScale = true;
 
@@ -2110,21 +2110,21 @@ void SFoliageEditMeshDisplayItem::OnUniformScale(ESlateCheckBoxState::Type InSta
 	}
 }
 
-ESlateCheckBoxState::Type SFoliageEditMeshDisplayItem::IsUniformScaleChecked() const
+ECheckBoxState SFoliageEditMeshDisplayItem::IsUniformScaleChecked() const
 {
-	return FoliageSettingsPtr->UniformScale ? ESlateCheckBoxState::Checked : ESlateCheckBoxState::Unchecked;
+	return FoliageSettingsPtr->UniformScale ? ECheckBoxState::Checked : ECheckBoxState::Unchecked;
 }
 
-void SFoliageEditMeshDisplayItem::OnScaleUniformReapply(ESlateCheckBoxState::Type InState)
+void SFoliageEditMeshDisplayItem::OnScaleUniformReapply(ECheckBoxState InState)
 {
-	FoliageSettingsPtr->ReapplyScaleX = InState == ESlateCheckBoxState::Checked ? true : false;
-	FoliageSettingsPtr->ReapplyScaleY = InState == ESlateCheckBoxState::Checked ? true : false;
-	FoliageSettingsPtr->ReapplyScaleZ = InState == ESlateCheckBoxState::Checked ? true : false;
+	FoliageSettingsPtr->ReapplyScaleX = InState == ECheckBoxState::Checked ? true : false;
+	FoliageSettingsPtr->ReapplyScaleY = InState == ECheckBoxState::Checked ? true : false;
+	FoliageSettingsPtr->ReapplyScaleZ = InState == ECheckBoxState::Checked ? true : false;
 }
 
-ESlateCheckBoxState::Type SFoliageEditMeshDisplayItem::IsScaleUniformReapplyChecked() const
+ECheckBoxState SFoliageEditMeshDisplayItem::IsScaleUniformReapplyChecked() const
 {
-	return FoliageSettingsPtr->ReapplyScaleX ? ESlateCheckBoxState::Checked : ESlateCheckBoxState::Unchecked;
+	return FoliageSettingsPtr->ReapplyScaleX ? ECheckBoxState::Checked : ECheckBoxState::Unchecked;
 }
 
 void SFoliageEditMeshDisplayItem::OnScaleUniformMinChanged(float InValue)
@@ -2159,14 +2159,14 @@ float SFoliageEditMeshDisplayItem::GetScaleUniformMax() const
 	return FoliageSettingsPtr->ScaleMaxX;
 }
 
-void SFoliageEditMeshDisplayItem::OnScaleXReapply(ESlateCheckBoxState::Type InState)
+void SFoliageEditMeshDisplayItem::OnScaleXReapply(ECheckBoxState InState)
 {
-	FoliageSettingsPtr->ReapplyScaleX = InState == ESlateCheckBoxState::Checked;
+	FoliageSettingsPtr->ReapplyScaleX = InState == ECheckBoxState::Checked;
 }
 
-ESlateCheckBoxState::Type SFoliageEditMeshDisplayItem::IsScaleXReapplyChecked() const
+ECheckBoxState SFoliageEditMeshDisplayItem::IsScaleXReapplyChecked() const
 {
-	return FoliageSettingsPtr->ReapplyScaleX ? ESlateCheckBoxState::Checked : ESlateCheckBoxState::Unchecked;
+	return FoliageSettingsPtr->ReapplyScaleX ? ECheckBoxState::Checked : ECheckBoxState::Unchecked;
 }
 
 void SFoliageEditMeshDisplayItem::OnScaleXMinChanged(float InValue)
@@ -2193,24 +2193,24 @@ float SFoliageEditMeshDisplayItem::GetScaleXMax() const
 	return FoliageSettingsPtr->ScaleMaxX;
 }
 
-ESlateCheckBoxState::Type SFoliageEditMeshDisplayItem::IsScaleXLockedChecked() const
+ECheckBoxState SFoliageEditMeshDisplayItem::IsScaleXLockedChecked() const
 {
-	return FoliageSettingsPtr->LockScaleX ? ESlateCheckBoxState::Checked : ESlateCheckBoxState::Unchecked;
+	return FoliageSettingsPtr->LockScaleX ? ECheckBoxState::Checked : ECheckBoxState::Unchecked;
 }
 
-void SFoliageEditMeshDisplayItem::OnScaleXLocked(ESlateCheckBoxState::Type InState)
+void SFoliageEditMeshDisplayItem::OnScaleXLocked(ECheckBoxState InState)
 {
-	FoliageSettingsPtr->LockScaleX = InState == ESlateCheckBoxState::Checked;
+	FoliageSettingsPtr->LockScaleX = InState == ECheckBoxState::Checked;
 }
 
-void SFoliageEditMeshDisplayItem::OnScaleYReapply(ESlateCheckBoxState::Type InState)
+void SFoliageEditMeshDisplayItem::OnScaleYReapply(ECheckBoxState InState)
 {
-	FoliageSettingsPtr->ReapplyScaleY = InState == ESlateCheckBoxState::Checked;
+	FoliageSettingsPtr->ReapplyScaleY = InState == ECheckBoxState::Checked;
 }
 
-ESlateCheckBoxState::Type SFoliageEditMeshDisplayItem::IsScaleYReapplyChecked() const
+ECheckBoxState SFoliageEditMeshDisplayItem::IsScaleYReapplyChecked() const
 {
-	return FoliageSettingsPtr->ReapplyScaleY ? ESlateCheckBoxState::Checked : ESlateCheckBoxState::Unchecked;
+	return FoliageSettingsPtr->ReapplyScaleY ? ECheckBoxState::Checked : ECheckBoxState::Unchecked;
 }
 
 void SFoliageEditMeshDisplayItem::OnScaleYMinChanged(float InValue)
@@ -2237,24 +2237,24 @@ float SFoliageEditMeshDisplayItem::GetScaleYMax() const
 	return FoliageSettingsPtr->ScaleMaxY;
 }
 
-ESlateCheckBoxState::Type SFoliageEditMeshDisplayItem::IsScaleYLockedChecked() const
+ECheckBoxState SFoliageEditMeshDisplayItem::IsScaleYLockedChecked() const
 {
-	return FoliageSettingsPtr->LockScaleY ? ESlateCheckBoxState::Checked : ESlateCheckBoxState::Unchecked;
+	return FoliageSettingsPtr->LockScaleY ? ECheckBoxState::Checked : ECheckBoxState::Unchecked;
 }
 
-void SFoliageEditMeshDisplayItem::OnScaleYLocked(ESlateCheckBoxState::Type InState)
+void SFoliageEditMeshDisplayItem::OnScaleYLocked(ECheckBoxState InState)
 {
-	FoliageSettingsPtr->LockScaleY = InState == ESlateCheckBoxState::Checked;
+	FoliageSettingsPtr->LockScaleY = InState == ECheckBoxState::Checked;
 }
 
-void SFoliageEditMeshDisplayItem::OnScaleZReapply(ESlateCheckBoxState::Type InState)
+void SFoliageEditMeshDisplayItem::OnScaleZReapply(ECheckBoxState InState)
 {
-	FoliageSettingsPtr->ReapplyScaleZ = InState == ESlateCheckBoxState::Checked ? true : false;
+	FoliageSettingsPtr->ReapplyScaleZ = InState == ECheckBoxState::Checked ? true : false;
 }
 
-ESlateCheckBoxState::Type SFoliageEditMeshDisplayItem::IsScaleZReapplyChecked() const
+ECheckBoxState SFoliageEditMeshDisplayItem::IsScaleZReapplyChecked() const
 {
-	return FoliageSettingsPtr->ReapplyScaleZ ? ESlateCheckBoxState::Checked : ESlateCheckBoxState::Unchecked;
+	return FoliageSettingsPtr->ReapplyScaleZ ? ECheckBoxState::Checked : ECheckBoxState::Unchecked;
 }
 
 void SFoliageEditMeshDisplayItem::OnScaleZMinChanged(float InValue)
@@ -2281,24 +2281,24 @@ float SFoliageEditMeshDisplayItem::GetScaleZMax() const
 	return FoliageSettingsPtr->ScaleMaxZ;
 }
 
-ESlateCheckBoxState::Type SFoliageEditMeshDisplayItem::IsScaleZLockedChecked() const
+ECheckBoxState SFoliageEditMeshDisplayItem::IsScaleZLockedChecked() const
 {
-	return FoliageSettingsPtr->LockScaleZ ? ESlateCheckBoxState::Checked : ESlateCheckBoxState::Unchecked;
+	return FoliageSettingsPtr->LockScaleZ ? ECheckBoxState::Checked : ECheckBoxState::Unchecked;
 }
 
-void SFoliageEditMeshDisplayItem::OnScaleZLocked(ESlateCheckBoxState::Type InState)
+void SFoliageEditMeshDisplayItem::OnScaleZLocked(ECheckBoxState InState)
 {
-	FoliageSettingsPtr->LockScaleZ = InState == ESlateCheckBoxState::Checked;
+	FoliageSettingsPtr->LockScaleZ = InState == ECheckBoxState::Checked;
 }
 
-void SFoliageEditMeshDisplayItem::OnZOffsetReapply(ESlateCheckBoxState::Type InState)
+void SFoliageEditMeshDisplayItem::OnZOffsetReapply(ECheckBoxState InState)
 {
-	FoliageSettingsPtr->ReapplyZOffset = InState == ESlateCheckBoxState::Checked ? true : false;
+	FoliageSettingsPtr->ReapplyZOffset = InState == ECheckBoxState::Checked ? true : false;
 }
 
-ESlateCheckBoxState::Type SFoliageEditMeshDisplayItem::IsZOffsetReapplyChecked() const
+ECheckBoxState SFoliageEditMeshDisplayItem::IsZOffsetReapplyChecked() const
 {
-	return FoliageSettingsPtr->ReapplyZOffset ? ESlateCheckBoxState::Checked : ESlateCheckBoxState::Unchecked;
+	return FoliageSettingsPtr->ReapplyZOffset ? ECheckBoxState::Checked : ECheckBoxState::Unchecked;
 }
 
 void SFoliageEditMeshDisplayItem::OnZOffsetMin(float InValue)
@@ -2325,14 +2325,14 @@ float SFoliageEditMeshDisplayItem::GetZOffsetMax() const
 	return FoliageSettingsPtr->ZOffsetMax;
 }
 
-void SFoliageEditMeshDisplayItem::OnRandomPitchReapply(ESlateCheckBoxState::Type InState)
+void SFoliageEditMeshDisplayItem::OnRandomPitchReapply(ECheckBoxState InState)
 {
-	FoliageSettingsPtr->ReapplyRandomPitchAngle = InState == ESlateCheckBoxState::Checked ? true : false;
+	FoliageSettingsPtr->ReapplyRandomPitchAngle = InState == ECheckBoxState::Checked ? true : false;
 }
 
-ESlateCheckBoxState::Type SFoliageEditMeshDisplayItem::IsRandomPitchReapplyChecked() const
+ECheckBoxState SFoliageEditMeshDisplayItem::IsRandomPitchReapplyChecked() const
 {
-	return FoliageSettingsPtr->ReapplyRandomPitchAngle ? ESlateCheckBoxState::Checked : ESlateCheckBoxState::Unchecked;
+	return FoliageSettingsPtr->ReapplyRandomPitchAngle ? ECheckBoxState::Checked : ECheckBoxState::Unchecked;
 }
 
 void SFoliageEditMeshDisplayItem::OnRandomPitchChanged(float InValue)
@@ -2345,14 +2345,14 @@ float SFoliageEditMeshDisplayItem::GetRandomPitch() const
 	return FoliageSettingsPtr->RandomPitchAngle;
 }
 
-void SFoliageEditMeshDisplayItem::OnGroundSlopeReapply(ESlateCheckBoxState::Type InState)
+void SFoliageEditMeshDisplayItem::OnGroundSlopeReapply(ECheckBoxState InState)
 {
-	FoliageSettingsPtr->ReapplyGroundSlope = InState == ESlateCheckBoxState::Checked ? true : false;
+	FoliageSettingsPtr->ReapplyGroundSlope = InState == ECheckBoxState::Checked ? true : false;
 }
 
-ESlateCheckBoxState::Type SFoliageEditMeshDisplayItem::IsGroundSlopeReapplyChecked() const
+ECheckBoxState SFoliageEditMeshDisplayItem::IsGroundSlopeReapplyChecked() const
 {
-	return FoliageSettingsPtr->ReapplyGroundSlope ? ESlateCheckBoxState::Checked : ESlateCheckBoxState::Unchecked;
+	return FoliageSettingsPtr->ReapplyGroundSlope ? ECheckBoxState::Checked : ECheckBoxState::Unchecked;
 }
 
 void SFoliageEditMeshDisplayItem::OnGroundSlopeChanged(float InValue)
@@ -2365,14 +2365,14 @@ float SFoliageEditMeshDisplayItem::GetGroundSlope() const
 	return FoliageSettingsPtr->GroundSlope;
 }
 
-void SFoliageEditMeshDisplayItem::OnHeightReapply(ESlateCheckBoxState::Type InState)
+void SFoliageEditMeshDisplayItem::OnHeightReapply(ECheckBoxState InState)
 {
-	FoliageSettingsPtr->ReapplyHeight = InState == ESlateCheckBoxState::Checked ? true : false;
+	FoliageSettingsPtr->ReapplyHeight = InState == ECheckBoxState::Checked ? true : false;
 }
 
-ESlateCheckBoxState::Type SFoliageEditMeshDisplayItem::IsHeightReapplyChecked() const
+ECheckBoxState SFoliageEditMeshDisplayItem::IsHeightReapplyChecked() const
 {
-	return FoliageSettingsPtr->ReapplyHeight ? ESlateCheckBoxState::Checked : ESlateCheckBoxState::Unchecked;
+	return FoliageSettingsPtr->ReapplyHeight ? ECheckBoxState::Checked : ECheckBoxState::Unchecked;
 }
 
 void SFoliageEditMeshDisplayItem::OnHeightMinChanged(float InValue)
@@ -2399,14 +2399,14 @@ float SFoliageEditMeshDisplayItem::GetHeightMax() const
 	return FoliageSettingsPtr->HeightMax;
 }
 
-void SFoliageEditMeshDisplayItem::OnLandscapeLayerReapply(ESlateCheckBoxState::Type InState)
+void SFoliageEditMeshDisplayItem::OnLandscapeLayerReapply(ECheckBoxState InState)
 {
-	FoliageSettingsPtr->ReapplyLandscapeLayer = InState == ESlateCheckBoxState::Checked ? true : false;
+	FoliageSettingsPtr->ReapplyLandscapeLayer = InState == ECheckBoxState::Checked ? true : false;
 }
 
-ESlateCheckBoxState::Type SFoliageEditMeshDisplayItem::IsLandscapeLayerReapplyChecked() const
+ECheckBoxState SFoliageEditMeshDisplayItem::IsLandscapeLayerReapplyChecked() const
 {
-	return FoliageSettingsPtr->ReapplyLandscapeLayer ? ESlateCheckBoxState::Checked : ESlateCheckBoxState::Unchecked;
+	return FoliageSettingsPtr->ReapplyLandscapeLayer ? ECheckBoxState::Checked : ECheckBoxState::Unchecked;
 }
 
 void SFoliageEditMeshDisplayItem::OnLandscapeLayerChanged(const FText& InValue)
@@ -2419,9 +2419,9 @@ FText SFoliageEditMeshDisplayItem::GetLandscapeLayer() const
 	return FText::FromName(FoliageSettingsPtr->LandscapeLayer);
 }
 
-void SFoliageEditMeshDisplayItem::OnCollisionWithWorld(ESlateCheckBoxState::Type InState)
+void SFoliageEditMeshDisplayItem::OnCollisionWithWorld(ECheckBoxState InState)
 {
-	if (InState == ESlateCheckBoxState::Checked)
+	if (InState == ECheckBoxState::Checked)
 	{
 		FoliageSettingsPtr->CollisionWithWorld = true;
 	}
@@ -2431,9 +2431,9 @@ void SFoliageEditMeshDisplayItem::OnCollisionWithWorld(ESlateCheckBoxState::Type
 	}
 }
 
-ESlateCheckBoxState::Type SFoliageEditMeshDisplayItem::IsCollisionWithWorldChecked() const
+ECheckBoxState SFoliageEditMeshDisplayItem::IsCollisionWithWorldChecked() const
 {
-	return FoliageSettingsPtr->CollisionWithWorld ? ESlateCheckBoxState::Checked : ESlateCheckBoxState::Unchecked;
+	return FoliageSettingsPtr->CollisionWithWorld ? ECheckBoxState::Checked : ECheckBoxState::Unchecked;
 }
 
 EVisibility SFoliageEditMeshDisplayItem::IsCollisionWithWorldVisible() const
@@ -2441,14 +2441,14 @@ EVisibility SFoliageEditMeshDisplayItem::IsCollisionWithWorldVisible() const
 	return FoliageSettingsPtr->CollisionWithWorld ? EVisibility::Visible : EVisibility::Collapsed;
 }
 
-void SFoliageEditMeshDisplayItem::OnCollisionWithWorldReapply(ESlateCheckBoxState::Type InState)
+void SFoliageEditMeshDisplayItem::OnCollisionWithWorldReapply(ECheckBoxState InState)
 {
-	FoliageSettingsPtr->ReapplyCollisionWithWorld = InState == ESlateCheckBoxState::Checked ? true : false;
+	FoliageSettingsPtr->ReapplyCollisionWithWorld = InState == ECheckBoxState::Checked ? true : false;
 }
 
-ESlateCheckBoxState::Type SFoliageEditMeshDisplayItem::IsCollisionWithWorldReapplyChecked() const
+ECheckBoxState SFoliageEditMeshDisplayItem::IsCollisionWithWorldReapplyChecked() const
 {
-	return FoliageSettingsPtr->ReapplyCollisionWithWorld ? ESlateCheckBoxState::Checked : ESlateCheckBoxState::Unchecked;
+	return FoliageSettingsPtr->ReapplyCollisionWithWorld ? ECheckBoxState::Checked : ECheckBoxState::Unchecked;
 }
 
 void SFoliageEditMeshDisplayItem::OnCollisionScaleXChanged(float InValue)
@@ -2481,14 +2481,14 @@ float SFoliageEditMeshDisplayItem::GetCollisionScaleZ() const
 	return FoliageSettingsPtr->CollisionScale.Z;
 }
 
-void SFoliageEditMeshDisplayItem::OnVertexColorMask(ESlateCheckBoxState::Type InState, FoliageVertexColorMask Mask)
+void SFoliageEditMeshDisplayItem::OnVertexColorMask(ECheckBoxState InState, FoliageVertexColorMask Mask)
 {
 	FoliageSettingsPtr->VertexColorMask = Mask;
 }
 
-ESlateCheckBoxState::Type SFoliageEditMeshDisplayItem::IsVertexColorMaskChecked(FoliageVertexColorMask Mask) const
+ECheckBoxState SFoliageEditMeshDisplayItem::IsVertexColorMaskChecked(FoliageVertexColorMask Mask) const
 {
-	return FoliageSettingsPtr->VertexColorMask == Mask ? ESlateCheckBoxState::Checked : ESlateCheckBoxState::Unchecked;
+	return FoliageSettingsPtr->VertexColorMask == Mask ? ECheckBoxState::Checked : ECheckBoxState::Unchecked;
 }
 
 void SFoliageEditMeshDisplayItem::OnVertexColorMaskThresholdChanged(float InValue)
@@ -2506,24 +2506,24 @@ EVisibility SFoliageEditMeshDisplayItem::IsVertexColorMaskThresholdVisible() con
 	return (FoliageSettingsPtr->VertexColorMask == FOLIAGEVERTEXCOLORMASK_Disabled) ? EVisibility::Collapsed : EVisibility::Visible;
 }
 
-void SFoliageEditMeshDisplayItem::OnVertexColorMaskInvert(ESlateCheckBoxState::Type InState)
+void SFoliageEditMeshDisplayItem::OnVertexColorMaskInvert(ECheckBoxState InState)
 {
-	FoliageSettingsPtr->VertexColorMaskInvert = (InState == ESlateCheckBoxState::Checked);
+	FoliageSettingsPtr->VertexColorMaskInvert = (InState == ECheckBoxState::Checked);
 }
 
-ESlateCheckBoxState::Type SFoliageEditMeshDisplayItem::IsVertexColorMaskInvertChecked() const
+ECheckBoxState SFoliageEditMeshDisplayItem::IsVertexColorMaskInvertChecked() const
 {
-	return FoliageSettingsPtr->VertexColorMaskInvert ? ESlateCheckBoxState::Checked : ESlateCheckBoxState::Unchecked;
+	return FoliageSettingsPtr->VertexColorMaskInvert ? ECheckBoxState::Checked : ECheckBoxState::Unchecked;
 }
 
-void SFoliageEditMeshDisplayItem::OnVertexColorMaskReapply(ESlateCheckBoxState::Type InState)
+void SFoliageEditMeshDisplayItem::OnVertexColorMaskReapply(ECheckBoxState InState)
 {
-	FoliageSettingsPtr->ReapplyVertexColorMask = InState == ESlateCheckBoxState::Checked ? true : false;
+	FoliageSettingsPtr->ReapplyVertexColorMask = InState == ECheckBoxState::Checked ? true : false;
 }
 
-ESlateCheckBoxState::Type SFoliageEditMeshDisplayItem::IsVertexColorMaskReapplyChecked() const
+ECheckBoxState SFoliageEditMeshDisplayItem::IsVertexColorMaskReapplyChecked() const
 {
-	return FoliageSettingsPtr->ReapplyVertexColorMask ? ESlateCheckBoxState::Checked : ESlateCheckBoxState::Unchecked;
+	return FoliageSettingsPtr->ReapplyVertexColorMask ? ECheckBoxState::Checked : ECheckBoxState::Unchecked;
 }
 
 FReply SFoliageEditMeshDisplayItem::OnMouseDownSelection(const FGeometry& MyGeometry, const FPointerEvent& MouseEvent)
