Fixed Linux CIS

diff --git a/Engine/Source/Editor/FoliageEdit/Private/FoliageEdMode.cpp b/Engine/Source/Editor/FoliageEdit/Private/FoliageEdMode.cpp
--- a/Engine/Source/Editor/FoliageEdit/Private/FoliageEdMode.cpp
+++ b/Engine/Source/Editor/FoliageEdit/Private/FoliageEdMode.cpp
@@ -1836,8 +1836,8 @@ bool FEdModeFoliage::InputKey(FEditorViewportClient* ViewportClient, FViewport*
 			if ((Viewport->IsPenActive() && Viewport->GetTabletPressure() > 0.f) ||
 				(!bMovingCamera && ViewportClient->GetCurrentWidgetAxis() == EAxisList::None &&
 					(FoliageEditorControlType == ELandscapeFoliageEditorControlType::IgnoreCtrl ||
-					 FoliageEditorControlType == ELandscapeFoliageEditorControlType::RequireCtrl   && IsCtrlDown(Viewport) ||
-					 FoliageEditorControlType == ELandscapeFoliageEditorControlType::RequireNoCtrl && !IsCtrlDown(Viewport))))
+					 (FoliageEditorControlType == ELandscapeFoliageEditorControlType::RequireCtrl   && IsCtrlDown(Viewport)) ||
+					 (FoliageEditorControlType == ELandscapeFoliageEditorControlType::RequireNoCtrl && !IsCtrlDown(Viewport)))))
 			{
 				if (!bToolActive)
 				{
