Customization of UFoliageType so that it has subcategories

diff --git a/Engine/Source/Editor/FoliageEdit/Private/FoliageTypeDetails.cpp b/Engine/Source/Editor/FoliageEdit/Private/FoliageTypeDetails.cpp
--- a/Engine/Source/Editor/FoliageEdit/Private/FoliageTypeDetails.cpp
+++ b/Engine/Source/Editor/FoliageEdit/Private/FoliageTypeDetails.cpp
@@ -10,8 +10,43 @@ TSharedRef<IDetailCustomization> FFoliageTypeDetails::MakeInstance()
 
 }
 
+void AddSubcategoryProperties(IDetailCategoryBuilder& CategoryBuilder)
+{
+	TArray<TSharedRef<IPropertyHandle>> CategoryProperties;
+	CategoryBuilder.GetDefaultProperties(CategoryProperties, true, true);
+
+	//Build map of subcategories to properties
+	const FName SubcategoryName("Subcategory");
+	TMap<FString, TArray<TSharedRef<IPropertyHandle>> > SubcategoryPropertiesMap;
+	for (auto& PropertyHandle : CategoryProperties)
+	{
+		if (UProperty* Property = PropertyHandle->GetProperty())
+		{
+			if (Property->HasMetaData(SubcategoryName))
+			{
+				const FString& Subcategory = Property->GetMetaData(SubcategoryName);
+				TArray<TSharedRef<IPropertyHandle>>& PropertyHandles = SubcategoryPropertiesMap.FindOrAdd(Subcategory);
+				PropertyHandles.Add(PropertyHandle);
+			}
+		}
+	}
+
+	//Add subgroups
+	for (auto Iter = SubcategoryPropertiesMap.CreateConstIterator(); Iter; ++Iter)
+	{
+		const FString &GroupString = Iter.Key();
+		const TArray<TSharedRef<IPropertyHandle>>& PropertyHandles = Iter.Value();
+		IDetailGroup& Group = CategoryBuilder.AddGroup(FName(*GroupString), FText::FromString(GroupString));
+		for (auto& PropertyHandle : PropertyHandles)
+		{
+			Group.AddPropertyRow(PropertyHandle);
+		}
+	}
+}
 
 void FFoliageTypeDetails::CustomizeDetails(IDetailLayoutBuilder& DetailBuilder)
 {
-	
+	AddSubcategoryProperties(DetailBuilder.EditCategory("General"));
+	AddSubcategoryProperties(DetailBuilder.EditCategory("Ecosystem"));
+	AddSubcategoryProperties(DetailBuilder.EditCategory("Painting"));
 }
\ No newline at end of file
