Fix non-unity build

diff --git a/Engine/Source/Editor/FoliageEdit/Private/FoliageTypeDetails.cpp b/Engine/Source/Editor/FoliageEdit/Private/FoliageTypeDetails.cpp
--- a/Engine/Source/Editor/FoliageEdit/Private/FoliageTypeDetails.cpp
+++ b/Engine/Source/Editor/FoliageEdit/Private/FoliageTypeDetails.cpp
@@ -2,6 +2,7 @@
 
 #include "UnrealEd.h"
 #include "FoliageEditModule.h"
+#include "PropertyEditing.h"
 #include "FoliageTypeDetails.h"
 
 TSharedRef<IDetailCustomization> FFoliageTypeDetails::MakeInstance()
