Fix foliage UI being hidden because of new subcategories

diff --git a/Engine/Source/Editor/FoliageEdit/Private/SFoliageEditMeshDisplayItem.cpp b/Engine/Source/Editor/FoliageEdit/Private/SFoliageEditMeshDisplayItem.cpp
--- a/Engine/Source/Editor/FoliageEdit/Private/SFoliageEditMeshDisplayItem.cpp
+++ b/Engine/Source/Editor/FoliageEdit/Private/SFoliageEditMeshDisplayItem.cpp
@@ -20,6 +20,27 @@
 
 #define LOCTEXT_NAMESPACE "FoliageEd_Mode"
 
+/** This is purposely empty because for now we just want a flat list of items. In the future we should break this out into a proper customization */
+class EmptyCustomization : public IDetailCustomization
+{
+public:
+	virtual ~EmptyCustomization(){};
+
+	/** Makes a new instance of this detail layout class for a specific detail view requesting it */
+	static TSharedRef<IDetailCustomization> MakeInstance()
+	{
+		return MakeShareable(new EmptyCustomization());
+	}
+
+	/** IDetailCustomization interface */
+	virtual void CustomizeDetails(IDetailLayoutBuilder& DetailBuilder)
+	{
+
+	}
+
+};
+
+
 BEGIN_SLATE_FUNCTION_BUILD_OPTIMIZATION
 void SFoliageEditMeshDisplayItem::Construct(const FArguments& InArgs)
 {
@@ -43,10 +64,12 @@ void SFoliageEditMeshDisplayItem::Construct(const FArguments& InArgs)
 	FPropertyEditorModule& PropertyEditorModule = FModuleManager::GetModuleChecked<FPropertyEditorModule>("PropertyEditor");
 	const FDetailsViewArgs DetailsViewArgs(false, false, true, false, true, this);
 	ClusterSettingsDetails = PropertyEditorModule.CreateDetailView(DetailsViewArgs);
+	ClusterSettingsDetails->RegisterInstancedCustomPropertyLayout(UFoliageType::StaticClass(), FOnGetDetailCustomizationInstance::CreateStatic(&EmptyCustomization::MakeInstance));
 
 	DetailsObjectList.Add(FoliageSettingsPtr);
 	ClusterSettingsDetails->SetIsPropertyVisibleDelegate(FIsPropertyVisible::CreateSP(this, &SFoliageEditMeshDisplayItem::IsPropertyVisible));
 	ClusterSettingsDetails->SetObjects(DetailsObjectList);
+	
 
 	// Everything (or almost) uses this padding, change it to expand the padding.
 	FMargin StandardPadding(2.0f, 2.0f, 2.0f, 2.0f);
@@ -1935,7 +1958,17 @@ FText SFoliageEditMeshDisplayItem::GetSaveRemoveSettingsTooltip() const
 bool SFoliageEditMeshDisplayItem::IsPropertyVisible(const FPropertyAndParent& PropertyAndParent) const
 {
 	const FString Category = FObjectEditorUtils::GetCategory(&PropertyAndParent.Property);
-	return Category == TEXT("Clustering") || Category == TEXT("Culling") || Category == TEXT("Lighting") || Category == TEXT("Collision") || Category == TEXT("Rendering");
+	const FName CategoryName = FObjectEditorUtils::GetCategoryFName(&PropertyAndParent.Property);
+	
+	FString Subcategory = "";
+	if ( (CategoryName != FName("Ecosystem")) && PropertyAndParent.Property.HasMetaData(FName("Subcategory")))
+	{
+		Subcategory = PropertyAndParent.Property.GetMetaData(FName("Subcategory"));
+	}
+
+	const FString Tag = Subcategory == "" ? Category : Subcategory;
+	
+	return Tag == TEXT("Clustering") || Tag == TEXT("Culling") || Tag == TEXT("Lighting") || Tag == TEXT("Collision") || Tag == TEXT("Rendering");
 }
 
 FText SFoliageEditMeshDisplayItem::GetInstanceCountString() const
