Fix non-unity build

diff --git a/Engine/Source/Editor/FoliageEdit/Private/ProceduralFoliageComponentDetails.cpp b/Engine/Source/Editor/FoliageEdit/Private/ProceduralFoliageComponentDetails.cpp
--- a/Engine/Source/Editor/FoliageEdit/Private/ProceduralFoliageComponentDetails.cpp
+++ b/Engine/Source/Editor/FoliageEdit/Private/ProceduralFoliageComponentDetails.cpp
@@ -1,6 +1,7 @@
 // Copyright 1998-2015 Epic Games, Inc. All Rights Reserved.
 
 #include "UnrealEd.h"
+#include "PropertyEditing.h"
 #include "ProceduralFoliageComponentDetails.h"
 #include "ProceduralFoliage.h"
 #include "ProceduralFoliageComponent.h"
