Build Fix - Temporary workaround to fix the build.

diff --git a/Engine/Source/Editor/FoliageEdit/Private/ProceduralFoliageFactory.cpp b/Engine/Source/Editor/FoliageEdit/Private/ProceduralFoliageFactory.cpp
--- a/Engine/Source/Editor/FoliageEdit/Private/ProceduralFoliageFactory.cpp
+++ b/Engine/Source/Editor/FoliageEdit/Private/ProceduralFoliageFactory.cpp
@@ -22,5 +22,5 @@ UObject* UProceduralFoliageFactory::FactoryCreateNew(UClass* Class, UObject* InP
 
 uint32 UProceduralFoliageFactory::GetMenuCategories() const
 {
-	return EAssetTypeCategories::Foliage;
+	return EAssetTypeCategories::Misc;
 }
\ No newline at end of file
