Deleting procedural foliage removes the instances that belong to it

diff --git a/Engine/Source/Editor/FoliageEdit/Private/FoliageEditModule.cpp b/Engine/Source/Editor/FoliageEdit/Private/FoliageEditModule.cpp
--- a/Engine/Source/Editor/FoliageEdit/Private/FoliageEditModule.cpp
+++ b/Engine/Source/Editor/FoliageEdit/Private/FoliageEditModule.cpp
@@ -14,6 +14,8 @@ const FName FoliageEditAppIdentifier = FName(TEXT("FoliageEdApp"));
 #include "ProceduralFoliageComponentDetails.h"
 #include "ActorFactoryProceduralFoliage.h"
 #include "ComponentVisualizer.h"
+#include "ProceduralFoliageActor.h"
+#include "ProceduralFoliageComponent.h"
 
 /**
  * Foliage Edit Mode module
@@ -46,6 +48,10 @@ public:
 		// Actor Factories
 		UActorFactoryProceduralFoliage* ProceduralFoliageActorFactory = ConstructObject<UActorFactoryProceduralFoliage>(UActorFactoryProceduralFoliage::StaticClass());
 		GEditor->ActorFactories.Add(ProceduralFoliageActorFactory);
+
+#if WITH_EDITOR
+		SubscribeEvents();
+#endif
 	}
 
 	/**
@@ -53,6 +59,10 @@ public:
 	 */
 	virtual void ShutdownModule() override
 	{
+#if WITH_EDITOR
+		UnsubscribeEvents();
+#endif
+
 		FEditorModeRegistry::Get().UnregisterMode(FBuiltinEditorModes::EM_Foliage);
 
 		if (!UObjectInitialized())
@@ -68,6 +78,33 @@ public:
 			PropertyModule.NotifyCustomizationModuleChanged();
 		}
 	}
+
+#if WITH_EDITOR
+
+	void OnLevelActorDeleted(AActor* Actor)
+	{
+		if (AProceduralFoliageActor* PFA = Cast<AProceduralFoliageActor>(Actor))
+		{
+			if (UProceduralFoliageComponent* ProceduralComponent = PFA->ProceduralComponent)
+			{
+				ProceduralComponent->RemoveProceduralContent();
+			}
+		}
+	}
+
+	void SubscribeEvents()
+	{
+		GEngine->OnLevelActorDeleted().Remove(OnLevelActorDeletedDelegateHandle);
+		OnLevelActorDeletedDelegateHandle = GEngine->OnLevelActorDeleted().AddRaw(this, &FFoliageEditModule::OnLevelActorDeleted);
+	}
+
+	void UnsubscribeEvents()
+	{
+		GEngine->OnLevelActorDeleted().Remove(OnLevelActorDeletedDelegateHandle);
+	}
+
+	FDelegateHandle OnLevelActorDeletedDelegateHandle;
+#endif
 };
 
 IMPLEMENT_MODULE( FFoliageEditModule, FoliageEdit );
diff --git a/Engine/Source/Runtime/Foliage/Private/ProceduralFoliageComponent.cpp b/Engine/Source/Runtime/Foliage/Private/ProceduralFoliageComponent.cpp
--- a/Engine/Source/Runtime/Foliage/Private/ProceduralFoliageComponent.cpp
+++ b/Engine/Source/Runtime/Foliage/Private/ProceduralFoliageComponent.cpp
@@ -217,8 +217,17 @@ void UProceduralFoliageComponent::SpawnTiles()
 
 void UProceduralFoliageComponent::SpawnProceduralContent()
 {
+#if WITH_EDITOR
+	RemoveProceduralContent();
+	SpawnTiles();
+#endif
+}
+
+void UProceduralFoliageComponent::RemoveProceduralContent()
+{
 #if WITH_EDITOR
 	UWorld* World = GetWorld();
+
 	for (ULevel* Level : World->GetLevels())
 	{
 		if (Level)
@@ -227,8 +236,6 @@ void UProceduralFoliageComponent::SpawnProceduralContent()
 			IFA->DeleteInstancesForProceduralFoliageComponent(this);
 		}
 	}
-	
-	SpawnTiles();
 #endif
 }
 
diff --git a/Engine/Source/Runtime/Foliage/Public/ProceduralFoliageComponent.h b/Engine/Source/Runtime/Foliage/Public/ProceduralFoliageComponent.h
--- a/Engine/Source/Runtime/Foliage/Public/ProceduralFoliageComponent.h
+++ b/Engine/Source/Runtime/Foliage/Public/ProceduralFoliageComponent.h
@@ -33,6 +33,7 @@ class FOLIAGE_API UProceduralFoliageComponent : public USceneComponent
 	UProceduralFoliage* ProceduralFoliage;
 
 	void SpawnProceduralContent();
+	void RemoveProceduralContent();
 	const FGuid& GetProceduralGuid() const { return ProceduralGuid;  }
 
 private:
