Splitting up the instance placement code and the brush logic code in the foliage editor so we can re-use placement.
Also added minimum layer weight so that you can control culling more easily

diff --git a/Engine/Source/Editor/FoliageEdit/Private/FoliageEdMode.cpp b/Engine/Source/Editor/FoliageEdit/Private/FoliageEdMode.cpp
--- a/Engine/Source/Editor/FoliageEdit/Private/FoliageEdMode.cpp
+++ b/Engine/Source/Editor/FoliageEdit/Private/FoliageEdMode.cpp
@@ -31,6 +31,7 @@
 #include "Components/ModelComponent.h"
 #include "EngineUtils.h"
 #include "LandscapeDataAccess.h"
+#include "FoliageType_InstancedStaticMesh.h"
 
 #define LOCTEXT_NAMESPACE "FoliageEdMode"
 #define FOLIAGE_SNAP_TRACE (10000.f)
@@ -247,35 +248,6 @@ void FEdModeFoliage::Tick(FEditorViewportClient* ViewportClient, float DeltaTime
 	}
 }
 
-static bool FoliageTrace(UWorld* InWorld, AInstancedFoliageActor* IFA, FHitResult& OutHit, FVector InStart, FVector InEnd, FName InTraceTag, bool InbReturnFaceIndex = false)
-{
-	FCollisionQueryParams QueryParams(InTraceTag, true, IFA);
-	QueryParams.bReturnFaceIndex = InbReturnFaceIndex;
-
-	bool bResult = true;
-	while (true)
-	{
-		bResult = InWorld->LineTraceSingle(OutHit, InStart, InEnd, QueryParams, FCollisionObjectQueryParams(ECC_WorldStatic));
-		if (bResult)
-		{
-			// In the editor traces can hit "No Collision" type actors, so ugh.
-			FBodyInstance* BodyInstance = OutHit.Component->GetBodyInstance();
-			if (BodyInstance->GetCollisionEnabled() != ECollisionEnabled::QueryAndPhysics || BodyInstance->GetResponseToChannel(ECC_WorldStatic) != ECR_Block)
-			{
-				AActor* Actor = OutHit.Actor.Get();
-				if (Actor)
-				{
-					QueryParams.AddIgnoredActor(Actor);
-				}
-				InStart = OutHit.ImpactPoint;
-				continue;
-			}
-		}
-		break;
-	}
-	return bResult;
-}
-
 /** Trace under the mouse cursor and update brush position */
 void FEdModeFoliage::FoliageBrushTrace(FEditorViewportClient* ViewportClient, int32 MouseX, int32 MouseY)
 {
@@ -301,7 +273,7 @@ void FEdModeFoliage::FoliageBrushTrace(FEditorViewportClient* ViewportClient, in
 			UWorld* World = ViewportClient->GetWorld();
 			AInstancedFoliageActor* IFA = AInstancedFoliageActor::GetInstancedFoliageActorForCurrentLevel(World, false);
 			static FName NAME_FoliageBrush = FName(TEXT("FoliageBrush"));
-			if (FoliageTrace(World, IFA, Hit, Start, End, NAME_FoliageBrush))
+			if (AInstancedFoliageActor::FoliageTrace(World, Hit, FDesiredFoliageInstance(Start, End), IFA, NAME_FoliageBrush))
 			{
 				// Check filters
 				UPrimitiveComponent* PrimComp = Hit.Component.Get();
@@ -374,144 +346,19 @@ void FEdModeFoliage::GetRandomVectorInBrush(FVector& OutStart, FVector& OutEnd)
 	OutEnd = BrushLocation + UISettings.GetRadius() * (Ru * U + Rv * V + Rw * BrushTraceDirection);
 }
 
-
-// Number of buckets for layer weight histogram distribution.
-#define NUM_INSTANCE_BUCKETS 10
-
-bool CheckCollisionWithWorld(UFoliageType* Settings, const FFoliageInstance& Inst, const FVector& HitNormal, const FVector& HitLocation, UWorld* InWorld, AInstancedFoliageActor* InIFA)
-{
-	FMatrix InstTransform = Inst.GetInstanceWorldTransform().ToMatrixWithScale();
-	FVector LocalHit = InstTransform.InverseTransformPosition(HitLocation);
-
-	if (Settings->CollisionWithWorld)
-	{
-		// Check for overhanging ledge
-		{
-			FVector LocalSamplePos[4] = { FVector(Settings->LowBoundOriginRadius.Z, 0, 0),
-				FVector(-Settings->LowBoundOriginRadius.Z, 0, 0),
-				FVector(0, Settings->LowBoundOriginRadius.Z, 0),
-				FVector(0, -Settings->LowBoundOriginRadius.Z, 0)
-			};
-
-			for (uint32 i = 0; i < 4; ++i)
-			{
-				FHitResult Hit;
-				FVector SamplePos = InstTransform.TransformPosition(FVector(Settings->LowBoundOriginRadius.X, Settings->LowBoundOriginRadius.Y, 2.f) + LocalSamplePos[i]);
-				float WorldRadius = (Settings->LowBoundOriginRadius.Z + 2.f)*FMath::Max(Inst.DrawScale3D.X, Inst.DrawScale3D.Y);
-				FVector NormalVector = Settings->AlignToNormal ? HitNormal : FVector(0, 0, 1);
-				if (FoliageTrace(InWorld, InIFA, Hit, SamplePos, SamplePos - NormalVector*WorldRadius, NAME_None))
-				{
-					if (LocalHit.Z - Inst.ZOffset < Settings->LowBoundOriginRadius.Z)
-					{
-						continue;
-					}
-				}
-				return false;
-			}
-		}
-
-		// Check collision with Bounding Box
-		{
-			FBox MeshBox = Settings->MeshBounds.GetBox();
-			MeshBox.Min.Z = FMath::Min(MeshBox.Max.Z, LocalHit.Z + Settings->MeshBounds.BoxExtent.Z * 0.05f);
-			FBoxSphereBounds ShrinkBound(MeshBox);
-			FBoxSphereBounds WorldBound = ShrinkBound.TransformBy(InstTransform);
-			//::DrawDebugBox(InWorld, WorldBound.Origin, WorldBound.BoxExtent, FColor::Red, true, 10.f);
-			static FName NAME_FoliageCollisionWithWorld = FName(TEXT("FoliageCollisionWithWorld"));
-			if (InWorld->OverlapTest(WorldBound.Origin, FQuat(Inst.Rotation), FCollisionShape::MakeBox(ShrinkBound.BoxExtent * Inst.DrawScale3D * Settings->CollisionScale), FCollisionQueryParams(NAME_FoliageCollisionWithWorld, false), FCollisionObjectQueryParams(ECC_WorldStatic)))
-			{
-				return false;
-			}
-		}
-	}
-
-	return true;
-}
-
-// Struct to hold potential instances we've randomly sampled
-struct FPotentialInstance
-{
-	FVector HitLocation;
-	FVector HitNormal;
-	UPrimitiveComponent* HitComponent;
-	float HitWeight;
-
-	FPotentialInstance(FVector InHitLocation, FVector InHitNormal, UPrimitiveComponent* InHitComponent, float InHitWeight)
-		: HitLocation(InHitLocation)
-		, HitNormal(InHitNormal)
-		, HitComponent(InHitComponent)
-		, HitWeight(InHitWeight)
-	{}
-
-	bool PlaceInstance(UFoliageType* Settings, FFoliageInstance& Inst, UWorld* InWorld, AInstancedFoliageActor* InIFA, bool bSkipCollision = false)
-	{
-		if (Settings->UniformScale)
-		{
-			float Scale = Settings->ScaleMinX + FMath::FRand() * (Settings->ScaleMaxX - Settings->ScaleMinX);
-			Inst.DrawScale3D = FVector(Scale, Scale, Scale);
-		}
-		else
-		{
-			float LockRand = FMath::FRand();
-			Inst.DrawScale3D.X = Settings->ScaleMinX + (Settings->LockScaleX ? LockRand : FMath::FRand()) * (Settings->ScaleMaxX - Settings->ScaleMinX);
-			Inst.DrawScale3D.Y = Settings->ScaleMinY + (Settings->LockScaleY ? LockRand : FMath::FRand()) * (Settings->ScaleMaxY - Settings->ScaleMinY);
-			Inst.DrawScale3D.Z = Settings->ScaleMinZ + (Settings->LockScaleZ ? LockRand : FMath::FRand()) * (Settings->ScaleMaxZ - Settings->ScaleMinZ);
-		}
-
-		Inst.ZOffset = Settings->ZOffsetMin + FMath::FRand() * (Settings->ZOffsetMax - Settings->ZOffsetMin);
-
-		Inst.Location = HitLocation;
-
-		// Random yaw and optional random pitch up to the maximum
-		Inst.Rotation = FRotator(FMath::FRand() * Settings->RandomPitchAngle, 0.f, 0.f);
-
-		if (Settings->RandomYaw)
-		{
-			Inst.Rotation.Yaw = FMath::FRand() * 360.f;
-		}
-		else
-		{
-			Inst.Flags |= FOLIAGE_NoRandomYaw;
-		}
-
-		if (Settings->AlignToNormal)
-		{
-			Inst.AlignToNormal(HitNormal, Settings->AlignMaxAngle);
-		}
-
-		// Apply the Z offset in local space
-		if (FMath::Abs(Inst.ZOffset) > KINDA_SMALL_NUMBER)
-		{
-			Inst.Location = Inst.GetInstanceWorldTransform().TransformPosition(FVector(0, 0, Inst.ZOffset));
-		}
-
-		UModelComponent* ModelComponent = Cast<UModelComponent>(HitComponent);
-		if (ModelComponent)
-		{
-			ABrush* BrushActor = ModelComponent->GetModel()->FindBrush(HitLocation);
-			if (BrushActor)
-			{
-				HitComponent = BrushActor->GetBrushComponent();
-			}
-		}
-
-		Inst.Base = HitComponent;
-
-		return bSkipCollision || CheckCollisionWithWorld(Settings, Inst, HitNormal, HitLocation, InWorld, InIFA);
-	}
-};
-
 static bool CheckLocationForPotentialInstance(FFoliageMeshInfo& MeshInfo, const UFoliageType* Settings, float DensityCheckRadius, const FVector& Location, const FVector& Normal, TArray<FVector>& PotentialInstanceLocations, FFoliageInstanceHash& PotentialInstanceHash)
 {
-	// Check slope
-	if ((Settings->GroundSlope > 0.f && Normal.Z <= FMath::Cos(PI * Settings->GroundSlope / 180.f) - SMALL_NUMBER) ||
-		(Settings->GroundSlope < 0.f && Normal.Z >= FMath::Cos(-PI * Settings->GroundSlope / 180.f) + SMALL_NUMBER))
+
+	// Check height range
+	if (Location.Z < Settings->HeightMin || Location.Z > Settings->HeightMax)
 	{
 		return false;
 	}
 
-	// Check height range
-	if (Location.Z < Settings->HeightMin || Location.Z > Settings->HeightMax)
+	// Check slope
+	const float MaxNormalAngle = FMath::Cos(FMath::DegreesToRadians(Settings->GroundSlope));
+	const float MinNormalAngle = FMath::Cos(FMath::DegreesToRadians(Settings->MinGroundSlope));
+	if (MaxNormalAngle > Normal.Z || MinNormalAngle < Normal.Z)	//keep in mind Hit.ImpactNormal.Z is (0,0,1) dot normal. However, ground slope is with relation to plane vector, not plane normal - so we swap comparisons
 	{
 		return false;
 	}
@@ -546,6 +393,7 @@ static bool CheckLocationForPotentialInstance(FFoliageMeshInfo& MeshInfo, const
 
 	int32 PotentialIdx = PotentialInstanceLocations.Add(Location);
 	PotentialInstanceHash.InsertInstance(Location, PotentialIdx);
+
 	return true;
 }
 
@@ -588,6 +436,187 @@ static bool CheckVertexColor(const UFoliageType* Settings, const FColor& VertexC
 	return true;
 }
 
+void FEdModeFoliage::CalculatePotentialInstances(UWorld* InWorld, const AInstancedFoliageActor* IgnoreIFA, const UFoliageType* FoliageType, const TArray<FDesiredFoliageInstance>& DesiredInstances, TArray<FPotentialInstance> OutPotentialInstances[NUM_INSTANCE_BUCKETS], LandscapeLayerCacheData* LandscapeLayerCachesPtr)
+{
+	LandscapeLayerCacheData LocalCache;
+	LandscapeLayerCachesPtr = LandscapeLayerCachesPtr ? LandscapeLayerCachesPtr : &LocalCache;
+
+	const UFoliageType* Settings = FoliageType;
+
+	// Quick lookup of potential instance locations, used for overlapping check.
+	TArray<FVector> PotentialInstanceLocations;
+	FFoliageInstanceHash PotentialInstanceHash(7);	// use 128x128 cell size, things like brush radius are typically small
+	PotentialInstanceLocations.Empty(DesiredInstances.Num());
+
+	// Reserve space in buckets for a potential instances 
+	for (int32 BucketIdx = 0; BucketIdx < NUM_INSTANCE_BUCKETS; ++BucketIdx)
+	{
+		auto& Bucket = OutPotentialInstances[BucketIdx];
+		Bucket.Reserve(DesiredInstances.Num());
+	}
+
+	for (const FDesiredFoliageInstance& DesiredInst : DesiredInstances)
+	{
+		FHitResult Hit;
+		static FName NAME_AddFoliageInstances = FName(TEXT("AddFoliageInstances"));
+		if (AInstancedFoliageActor* IFA = AInstancedFoliageActor::FoliageTrace(InWorld, Hit, DesiredInst, IgnoreIFA, NAME_AddFoliageInstances, true))
+		{
+			// Check filters
+			UPrimitiveComponent* PrimComp = Hit.Component.Get();
+			UMaterialInterface* Material = PrimComp ? PrimComp->GetMaterial(0) : nullptr;
+
+			if (DesiredInst.PlacementMode == EFoliagePlacementMode::Manual)
+			{
+				if (!PrimComp ||
+					PrimComp->GetOutermost() != InWorld->GetCurrentLevel()->GetOutermost()/* ||
+					(!UISettings.bFilterLandscape && PrimComp->IsA(ULandscapeHeightfieldCollisionComponent::StaticClass())) ||
+					(!UISettings.bFilterStaticMesh && PrimComp->IsA(UStaticMeshComponent::StaticClass())) ||
+					(!UISettings.bFilterBSP && PrimComp->IsA(UModelComponent::StaticClass())) ||
+					(!UISettings.bFilterTranslucent && Material && IsTranslucentBlendMode(Material->GetBlendMode()))*/
+					)
+				{
+					continue;
+				}
+			}
+			else
+			{
+				if (PrimComp == false)
+				{
+					continue;
+				}
+			}
+			
+
+
+			FFoliageMeshInfo* MeshInfo = DesiredInst.MeshInfo;
+			if (DesiredInst.PlacementMode == EFoliagePlacementMode::Procedural)
+			{
+				//We have to add the mesh into the foliage actor. This assumes the foliage type is just a concrete InstancedStaticMesh.
+				//With the new tree system coming this will not be true and you'll need to traverse the tree and insert the needed meshes
+				if (UStaticMesh* StaticMesh = FoliageType->GetStaticMesh())
+				{
+					UFoliageType* ExistingSettings = IFA->GetSettingsForMesh(StaticMesh);
+					if (ExistingSettings == nullptr)
+					{
+						if (const UFoliageType_InstancedStaticMesh* InstancedStaticMeshFoliageType = Cast<UFoliageType_InstancedStaticMesh>(Settings))
+						{
+							MeshInfo = IFA->AddMesh(Settings->GetStaticMesh(), nullptr, InstancedStaticMeshFoliageType);
+						}
+					}
+					else
+					{
+						MeshInfo = IFA->FindMesh(ExistingSettings);
+						//Settings = ExistingSettings;
+					}
+				}
+
+				
+			}
+
+			if (MeshInfo == nullptr)
+			{
+				continue;
+			}
+
+			// Radius where we expect to have a single instance, given the density rules
+			const float DensityCheckRadius = FMath::Max<float>(FMath::Sqrt((1000.f*1000.f) / (PI * Settings->Density)), Settings->Radius);
+			const FName LandscapeLayerName = Settings->LandscapeLayer;
+			TMap<ULandscapeComponent*, TArray<uint8>>* LandscapeLayerCache = LandscapeLayerName != NAME_None ? &LandscapeLayerCachesPtr->FindOrAdd(LandscapeLayerName) : nullptr;
+			
+
+			if (CheckLocationForPotentialInstance(*MeshInfo, Settings, DensityCheckRadius, Hit.ImpactPoint, Hit.ImpactNormal, PotentialInstanceLocations, PotentialInstanceHash) == false)
+			{
+				continue;
+			}
+
+			// Check vertex color mask
+			if (Settings->VertexColorMask != FOLIAGEVERTEXCOLORMASK_Disabled && Hit.FaceIndex != INDEX_NONE)
+			{
+				if (UStaticMeshComponent* HitStaticMeshComponent = Cast<UStaticMeshComponent>(Hit.Component.Get()))
+				{
+					FColor VertexColor;
+					if (GetStaticMeshVertexColorForHit(HitStaticMeshComponent, Hit.FaceIndex, Hit.ImpactPoint, VertexColor))
+					{
+						if (!CheckVertexColor(Settings, VertexColor))
+						{
+							continue;
+						}
+					}
+				}
+			}
+
+			// Check landscape layer
+			float HitWeight = 1.f;
+			if (LandscapeLayerCache)
+			{
+				if (ULandscapeHeightfieldCollisionComponent* HitLandscapeCollision = Cast<ULandscapeHeightfieldCollisionComponent>(Hit.Component.Get()))
+				{
+					if (ULandscapeComponent* HitLandscape = HitLandscapeCollision->RenderComponent.Get())
+					{
+						TArray<uint8>* LayerCache = &LandscapeLayerCache->FindOrAdd(HitLandscape);
+						// TODO: FName to LayerInfo?
+						HitWeight = HitLandscape->GetLayerWeightAtLocation(Hit.ImpactPoint, HitLandscape->GetLandscapeInfo()->GetLayerInfoByName(LandscapeLayerName), LayerCache);
+
+						// Reject instance randomly in proportion to weight
+						const float WeightNeeded = FMath::Max(Settings->MinimumLayerWeight, FMath::FRand());
+						if (HitWeight < FMath::Max(SMALL_NUMBER, WeightNeeded))
+						{
+							continue;
+						}
+					}
+				}
+			}
+
+			const int32 BucketIndex = FMath::RoundToInt(HitWeight * (float)(NUM_INSTANCE_BUCKETS - 1));
+			OutPotentialInstances[BucketIndex].Add(FPotentialInstance(Hit.ImpactPoint, Hit.ImpactNormal, Hit.Component.Get(), HitWeight, MeshInfo, IFA, DesiredInst));
+		}
+	}
+}
+
+void FEdModeFoliage::AddInstances(UWorld* InWorld, const TArray<FDesiredFoliageInstance>& DesiredInstances)
+{
+	TMap<const UFoliageType*, TArray<FDesiredFoliageInstance>> SettingsInstancesMap;
+	for (const FDesiredFoliageInstance& DesiredInst : DesiredInstances)
+	{
+		TArray<FDesiredFoliageInstance>& Instances = SettingsInstancesMap.FindOrAdd(DesiredInst.FoliageType);
+		Instances.Add(DesiredInst);
+	}
+
+	for (auto It = SettingsInstancesMap.CreateConstIterator(); It; ++It)
+	{
+		const UFoliageType* FoliageType = It.Key();
+
+		const TArray<FDesiredFoliageInstance>& Instances = It.Value();
+		AddInstancesImp(InWorld, FoliageType, Instances);
+	}
+}
+
+void FEdModeFoliage::AddInstancesImp(UWorld* InWorld, const UFoliageType* Settings, const TArray<FDesiredFoliageInstance>& DesiredInstances, const AInstancedFoliageActor* IgnoreIFA, const TArray<int32>& ExistingInstanceBuckets, const float Pressure, LandscapeLayerCacheData* LandscapeLayerCachesPtr)
+{
+	TArray<FPotentialInstance> PotentialInstanceBuckets[NUM_INSTANCE_BUCKETS];
+	CalculatePotentialInstances(InWorld, IgnoreIFA, Settings, DesiredInstances, PotentialInstanceBuckets, LandscapeLayerCachesPtr);
+
+	for (int32 BucketIdx = 0; BucketIdx < NUM_INSTANCE_BUCKETS; BucketIdx++)
+	{
+		TArray<FPotentialInstance>& PotentialInstances = PotentialInstanceBuckets[BucketIdx];
+		float BucketFraction = (float)(BucketIdx + 1) / (float)NUM_INSTANCE_BUCKETS;
+
+		// We use the number that actually succeeded in placement (due to parameters) as the target
+		// for the number that should be in the brush region.
+		const int32 BucketOffset = (ExistingInstanceBuckets.Num() ? ExistingInstanceBuckets[BucketIdx] : 0);
+		int32 AdditionalInstances = FMath::Clamp<int32>(FMath::RoundToInt(BucketFraction * (float)(PotentialInstances.Num() - BucketOffset) * Pressure), 0, PotentialInstances.Num());
+		for (int32 Idx = 0; Idx < AdditionalInstances; Idx++)
+		{
+			FPotentialInstance& PotentialInstance = PotentialInstances[Idx];
+			FFoliageInstance Inst;
+			if (PotentialInstance.PlaceInstance(Settings, Inst))
+			{
+				Inst.ProceduralGuid = PotentialInstance.DesiredInstance.ProceduralGuid;
+				PotentialInstance.MeshInfo->AddInstance(PotentialInstance.IFA, Settings, Inst);
+			}
+		}
+	}
+}
 
 /** Add instances inside the brush to match DesiredInstanceCount */
 void FEdModeFoliage::AddInstancesForBrush(UWorld* InWorld, AInstancedFoliageActor* IFA, UFoliageType* Settings, FFoliageMeshInfo& MeshInfo, int32 DesiredInstanceCount, const TArray<int32>& ExistingInstances, float Pressure)
@@ -595,8 +624,8 @@ void FEdModeFoliage::AddInstancesForBrush(UWorld* InWorld, AInstancedFoliageActo
 	checkf(InWorld == IFA->GetWorld(), TEXT("Warning:World does not match Foliage world"));
 	if (DesiredInstanceCount > ExistingInstances.Num())
 	{
-		int32 ExistingInstanceBuckets[NUM_INSTANCE_BUCKETS];
-		FMemory::Memzero(ExistingInstanceBuckets, sizeof(ExistingInstanceBuckets));
+		TArray<int32> ExistingInstanceBuckets;
+		ExistingInstanceBuckets.AddUninitialized(NUM_INSTANCE_BUCKETS);
 
 		// Cache store mapping between component and weight data
 		TMap<ULandscapeComponent*, TArray<uint8> >* LandscapeLayerCache = nullptr;
@@ -632,116 +661,19 @@ void FEdModeFoliage::AddInstancesForBrush(UWorld* InWorld, AInstancedFoliageActo
 			ExistingInstanceBuckets[NUM_INSTANCE_BUCKETS - 1] = ExistingInstances.Num();
 		}
 
-		// We calculate a set of potential ExistingInstances for the brush area.
 		TArray<FPotentialInstance> PotentialInstanceBuckets[NUM_INSTANCE_BUCKETS];
-		// Reserve space in buckets for a potential instances 
-		for (auto& Bucket : PotentialInstanceBuckets)
-		{
-			Bucket.Reserve(DesiredInstanceCount);
-		}
-
-		// Quick lookup of potential instance locations, used for overlapping check.
-		TArray<FVector> PotentialInstanceLocations;
-		FFoliageInstanceHash PotentialInstanceHash(7);	// use 128x128 cell size, as the brush radius is typically small.
-		PotentialInstanceLocations.Empty(DesiredInstanceCount);
-
-		// Radius where we expect to have a single instance, given the density rules
-		const float DensityCheckRadius = FMath::Max<float>(FMath::Sqrt((1000.f*1000.f) / (PI * Settings->Density)), Settings->Radius);
+		TArray<FDesiredFoliageInstance> DesiredInstances;	//we compute instances for the brush
+		DesiredInstances.Reserve(DesiredInstanceCount);
 
 		for (int32 DesiredIdx = 0; DesiredIdx < DesiredInstanceCount; DesiredIdx++)
 		{
 			FVector Start, End;
-
 			GetRandomVectorInBrush(Start, End);
-
-			FHitResult Hit;
-			static FName NAME_AddInstancesForBrush = FName(TEXT("AddInstancesForBrush"));
-			if (FoliageTrace(InWorld, IFA, Hit, Start, End, NAME_AddInstancesForBrush, true))
-			{
-				// Check filters
-				UPrimitiveComponent* PrimComp = Hit.Component.Get();
-				UMaterialInterface* Material = PrimComp ? PrimComp->GetMaterial(0) : nullptr;
-
-				if (!PrimComp ||
-					PrimComp->GetOutermost() != InWorld->GetCurrentLevel()->GetOutermost() ||
-					(!UISettings.bFilterLandscape && PrimComp->IsA(ULandscapeHeightfieldCollisionComponent::StaticClass())) ||
-					(!UISettings.bFilterStaticMesh && PrimComp->IsA(UStaticMeshComponent::StaticClass())) ||
-					(!UISettings.bFilterBSP && PrimComp->IsA(UModelComponent::StaticClass())) ||
-					(!UISettings.bFilterTranslucent && Material && IsTranslucentBlendMode(Material->GetBlendMode()))
-					)
-				{
-					continue;
-				}
-
-				if (!CheckLocationForPotentialInstance(MeshInfo, Settings, DensityCheckRadius, Hit.Location, Hit.Normal, PotentialInstanceLocations, PotentialInstanceHash))
-				{
-					continue;
-				}
-
-				// Check vertex color mask
-				if (Settings->VertexColorMask != FOLIAGEVERTEXCOLORMASK_Disabled && Hit.FaceIndex != INDEX_NONE)
-				{
-					UStaticMeshComponent* HitStaticMeshComponent = Cast<UStaticMeshComponent>(Hit.Component.Get());
-					if (HitStaticMeshComponent)
-					{
-						FColor VertexColor;
-						if (GetStaticMeshVertexColorForHit(HitStaticMeshComponent, Hit.FaceIndex, Hit.Location, VertexColor))
-						{
-							if (!CheckVertexColor(Settings, VertexColor))
-							{
-								continue;
-							}
-						}
-					}
-				}
-
-				// Check landscape layer
-				float HitWeight = 1.f;
-				if (LandscapeLayerName != NAME_None)
-				{
-					ULandscapeHeightfieldCollisionComponent* HitLandscapeCollision = Cast<ULandscapeHeightfieldCollisionComponent>(Hit.Component.Get());
-					if (HitLandscapeCollision)
-					{
-						ULandscapeComponent* HitLandscape = HitLandscapeCollision->RenderComponent.Get();
-						if (HitLandscape)
-						{
-							TArray<uint8>* LayerCache = &LandscapeLayerCache->FindOrAdd(HitLandscape);
-							// TODO: FName to LayerInfo?
-							HitWeight = HitLandscape->GetLayerWeightAtLocation(Hit.Location, HitLandscape->GetLandscapeInfo()->GetLayerInfoByName(LandscapeLayerName), LayerCache);
-
-							// Reject instance randomly in proportion to weight
-							if (HitWeight <= FMath::FRand())
-							{
-								continue;
-							}
-						}
-					}
-				}
-
-				const int32 BucketIndex = FMath::RoundToInt(HitWeight * (float)(NUM_INSTANCE_BUCKETS - 1));
-				check(BucketIndex < ARRAY_COUNT(PotentialInstanceBuckets));
-
-				PotentialInstanceBuckets[BucketIndex].Add(FPotentialInstance(Hit.Location, Hit.Normal, Hit.Component.Get(), HitWeight));
-			}
+			FDesiredFoliageInstance* DesiredInstance = new (DesiredInstances)FDesiredFoliageInstance(Start, End);
+			DesiredInstance->MeshInfo = &MeshInfo;
 		}
 
-		for (int32 BucketIdx = 0; BucketIdx < NUM_INSTANCE_BUCKETS; BucketIdx++)
-		{
-			TArray<FPotentialInstance>& PotentialInstances = PotentialInstanceBuckets[BucketIdx];
-			float BucketFraction = (float)(BucketIdx + 1) / (float)NUM_INSTANCE_BUCKETS;
-
-			// We use the number that actually succeeded in placement (due to parameters) as the target
-			// for the number that should be in the brush region.
-			int32 AdditionalInstances = FMath::Clamp<int32>(FMath::RoundToInt(BucketFraction * (float)(PotentialInstances.Num() - ExistingInstanceBuckets[BucketIdx]) * Pressure), 0, PotentialInstances.Num());
-			for (int32 Idx = 0; Idx < AdditionalInstances; Idx++)
-			{
-				FFoliageInstance Inst;
-				if (PotentialInstances[Idx].PlaceInstance(Settings, Inst, InWorld, IFA))
-				{
-					MeshInfo.AddInstance(IFA, Settings, Inst);
-				}
-			}
-		}
+		AddInstancesImp(InWorld, Settings, DesiredInstances, IFA, ExistingInstanceBuckets, Pressure, &LandscapeLayerCaches);
 	}
 }
 
@@ -975,7 +907,7 @@ void FEdModeFoliage::ReapplyInstancesForBrush(UWorld* InWorld, AInstancedFoliage
 				FVector ZAxis = Instance.Rotation.Quaternion().GetAxisZ();
 				FVector Start = Instance.Location + 16.f * ZAxis;
 				FVector End = Instance.Location - 16.f * ZAxis;
-				if (FoliageTrace(InWorld, IFA, Hit, Start, End, NAME_ReapplyInstancesForBrush, true))
+				if (AInstancedFoliageActor::FoliageTrace(InWorld, Hit, FDesiredFoliageInstance(Start, End), IFA, NAME_ReapplyInstancesForBrush, true))
 				{
 					// Reapply the normal
 					if (bReapplyNormal)
@@ -1103,9 +1035,9 @@ void FEdModeFoliage::ReapplyInstancesForBrush(UWorld* InWorld, AInstancedFoliage
 				static const FName NAME_ReapplyInstancesForBrush = TEXT("ReapplyCollisionWithWorld");
 				FVector Start = Instance.Location + FVector(0.f, 0.f, 16.f);
 				FVector End = Instance.Location - FVector(0.f, 0.f, 16.f);
-				if (FoliageTrace(InWorld, IFA, Hit, Start, End, NAME_ReapplyInstancesForBrush))
+				if (AInstancedFoliageActor::FoliageTrace(InWorld, Hit, FDesiredFoliageInstance(Start, End), IFA, NAME_ReapplyInstancesForBrush))
 				{
-					if (!CheckCollisionWithWorld(Settings, Instance, Hit.Normal, Hit.Location, InWorld, IFA))
+					if (!IFA->CheckCollisionWithWorld(Settings, Instance, Hit.Normal, Hit.Location))
 					{
 						InstancesToDelete.Add(InstanceIndex);
 						continue;
@@ -1458,7 +1390,7 @@ void FEdModeFoliage::ApplyPaintBucket(AActor* Actor, bool bRemove)
 				for (int32 Idx = 0; Idx < InstancesToPlace.Num(); Idx++)
 				{
 					FFoliageInstance Inst;
-					if (InstancesToPlace[Idx].PlaceInstance(Settings, Inst, World, IFA))
+					if (InstancesToPlace[Idx].PlaceInstance(Settings, Inst))
 					{
 						MeshInfo.AddInstance(IFA, Settings, Inst);
 					}
@@ -1468,7 +1400,7 @@ void FEdModeFoliage::ApplyPaintBucket(AActor* Actor, bool bRemove)
 	}
 }
 
-bool FEdModeFoliage::GetStaticMeshVertexColorForHit(UStaticMeshComponent* InStaticMeshComponent, int32 InTriangleIndex, const FVector& InHitLocation, FColor& OutVertexColor) const
+bool FEdModeFoliage::GetStaticMeshVertexColorForHit(UStaticMeshComponent* InStaticMeshComponent, int32 InTriangleIndex, const FVector& InHitLocation, FColor& OutVertexColor)
 {
 	UStaticMesh* StaticMesh = InStaticMeshComponent->StaticMesh;
 
@@ -1940,7 +1872,7 @@ bool FEdModeFoliage::InputKey(FEditorViewportClient* ViewportClient, FViewport*
 						FVector End = Instance.Location - FVector(0.f, 0.f, FOLIAGE_SNAP_TRACE);
 
 						FHitResult Hit;
-						if (FoliageTrace(GetWorld(), IFA, Hit, Start, End, FName("FoliageSnap")))
+						if (AInstancedFoliageActor::FoliageTrace(GetWorld(), Hit, FDesiredFoliageInstance(Start, End), IFA, FName("FoliageSnap")))
 						{
 							// Check current level
 							if ((Hit.Component.IsValid() && Hit.Component.Get()->GetOutermost() == World->GetCurrentLevel()->GetOutermost()) ||
diff --git a/Engine/Source/Editor/FoliageEdit/Private/FoliageEdMode.h b/Engine/Source/Editor/FoliageEdit/Private/FoliageEdMode.h
--- a/Engine/Source/Editor/FoliageEdit/Private/FoliageEdMode.h
+++ b/Engine/Source/Editor/FoliageEdit/Private/FoliageEdMode.h
@@ -146,6 +146,8 @@ public:
 	}
 };
 
+// Number of buckets for layer weight histogram distribution.
+#define NUM_INSTANCE_BUCKETS 10
 
 /**
  * Foliage editor mode
@@ -301,12 +303,22 @@ public:
 	 */
 	UFoliageType* SaveSettingsObject(const FText& InSettingsPackageName, UFoliageType* Settings);
 
+	/** Add desired instances. Uses foliage settings to determine location/scale/rotation and whether instances should be ignored */
+	static void AddInstances(UWorld* InWorld, const TArray<FDesiredFoliageInstance>& DesiredInstances);
+
 	/** Apply paint bucket to actor */
 	void ApplyPaintBucket(AActor* Actor, bool bRemove);
 private:
+	typedef TMap<FName, TMap<ULandscapeComponent*, TArray<uint8> > > LandscapeLayerCacheData;
+
 	/** Add instances inside the brush to match DesiredInstanceCount */
 	void AddInstancesForBrush(UWorld* InWorld, AInstancedFoliageActor* IFA, UFoliageType* Settings, FFoliageMeshInfo& MeshInfo, int32 DesiredInstanceCount, const TArray<int32>& ExistingInstances, float Pressure);
 
+	/** Common code for adding instances to world based on settings */
+	static void AddInstancesImp(UWorld* InWorld, const UFoliageType* Settings, const TArray<FDesiredFoliageInstance>& DesiredInstances, const AInstancedFoliageActor* IgnoreIFA = nullptr, const TArray<int32>& ExistingInstances = TArray<int32>(), const float Pressure = 1.f, LandscapeLayerCacheData* LandscapeLayerCaches = nullptr);
+
+	static void CalculatePotentialInstances(UWorld* InWorld, const AInstancedFoliageActor* IFA, const UFoliageType* Settings, const TArray<FDesiredFoliageInstance>& DesiredInstances, TArray<FPotentialInstance> OutPotentialInstances[NUM_INSTANCE_BUCKETS], LandscapeLayerCacheData* LandscaleLayerCachesPtr);
+
 	/** Remove instances inside the brush to match DesiredInstanceCount. NOTE: PotentialInstancesToRemove array is modified by this function. */
 	void RemoveInstancesForBrush(AInstancedFoliageActor* IFA, FFoliageMeshInfo& MeshInfo, int32 DesiredInstanceCount, TArray<int32>& PotentialInstancesToRemove, float Pressure);
 
@@ -314,7 +326,7 @@ private:
 	void ReapplyInstancesForBrush(UWorld* InWorld, AInstancedFoliageActor* IFA, UFoliageType* Settings, const TArray<int32>& ExistingInstances);
 
 	/** Lookup the vertex color corresponding to a location traced on a static mesh */
-	bool GetStaticMeshVertexColorForHit(UStaticMeshComponent* InStaticMeshComponent, int32 InTriangleIndex, const FVector& InHitLocation, FColor& OutVertexColor) const;
+	static bool GetStaticMeshVertexColorForHit(UStaticMeshComponent* InStaticMeshComponent, int32 InTriangleIndex, const FVector& InHitLocation, FColor& OutVertexColor);
 
 	bool bBrushTraceValid;
 	FVector BrushLocation;
@@ -322,7 +334,7 @@ private:
 	UStaticMeshComponent* SphereBrushComponent;
 
 	// Landscape layer cache data
-	TMap<FName, TMap<ULandscapeComponent*, TArray<uint8> > > LandscapeLayerCaches;
+	LandscapeLayerCacheData LandscapeLayerCaches;
 
 	// Placed level data
 	TArray<FFoliageMeshUIInfo> FoliageMeshList;
diff --git a/Engine/Source/Editor/FoliageEdit/Private/ProceduralFoliageComponentDetails.cpp b/Engine/Source/Editor/FoliageEdit/Private/ProceduralFoliageComponentDetails.cpp
--- a/Engine/Source/Editor/FoliageEdit/Private/ProceduralFoliageComponentDetails.cpp
+++ b/Engine/Source/Editor/FoliageEdit/Private/ProceduralFoliageComponentDetails.cpp
@@ -5,6 +5,8 @@
 #include "ProceduralFoliageComponentDetails.h"
 #include "ProceduralFoliage.h"
 #include "ProceduralFoliageComponent.h"
+#include "InstancedFoliage.h"
+#include "FoliageEdMode.h"
 
 TSharedRef<IDetailCustomization> FProceduralFoliageComponentDetails::MakeInstance()
 {
@@ -72,7 +74,9 @@ FReply FProceduralFoliageComponentDetails::OnResimulateClicked()
 				Component->ProceduralFoliage->Simulate();
 			}
 
-			Component->SpawnProceduralContent();
+			TArray <FDesiredFoliageInstance> DesiredFoliageInstances;
+			Component->SpawnProceduralContent(DesiredFoliageInstances);
+			FEdModeFoliage::AddInstances(Component->GetWorld(), DesiredFoliageInstances);
 		}
 	}
 	return FReply::Handled();
diff --git a/Engine/Source/Runtime/Foliage/Private/InstancedFoliage.cpp b/Engine/Source/Runtime/Foliage/Private/InstancedFoliage.cpp
--- a/Engine/Source/Runtime/Foliage/Private/InstancedFoliage.cpp
+++ b/Engine/Source/Runtime/Foliage/Private/InstancedFoliage.cpp
@@ -13,6 +13,7 @@
 #include "Components/HierarchicalInstancedStaticMeshComponent.h"
 #include "Serialization/CustomVersion.h"
 #include "ProceduralFoliageComponent.h"
+#include "ProceduralFoliageBlockingVolume.h"
 
 #define LOCTEXT_NAMESPACE "InstancedFoliage"
 
@@ -195,6 +196,7 @@ UFoliageType::UFoliageType(const FObjectInitializer& ObjectInitializer)
 	ZOffsetMin = 0.0f;
 	ZOffsetMax = 0.0f;
 	LandscapeLayer = NAME_None;
+	MinimumLayerWeight = 0.5f;
 	DisplayOrder = 0;
 	IsSelected = false;
 	ShowNothing = false;
@@ -396,7 +398,7 @@ void FFoliageMeshInfo::CheckValid()
 #endif
 }
 
-void FFoliageMeshInfo::AddInstance(AInstancedFoliageActor* InIFA, UFoliageType* InSettings, const FFoliageInstance& InNewInstance)
+void FFoliageMeshInfo::AddInstance(AInstancedFoliageActor* InIFA, const UFoliageType* InSettings, const FFoliageInstance& InNewInstance)
 {
 	InIFA->Modify();
 
@@ -1818,4 +1820,187 @@ void AInstancedFoliageActor::ApplyWorldOffset(const FVector& InOffset, bool bWor
 	}
 }
 
+AInstancedFoliageActor* AInstancedFoliageActor::FoliageTrace(UWorld* InWorld, FHitResult& OutHit, const FDesiredFoliageInstance& DesiredInstance, const AInstancedFoliageActor* IgnoreIFA, FName InTraceTag, bool InbReturnFaceIndex)
+{
+	FCollisionQueryParams QueryParams(InTraceTag, true, IgnoreIFA);
+	QueryParams.bReturnFaceIndex = InbReturnFaceIndex;
+
+	FVector StartTrace = DesiredInstance.StartTrace;
+
+	TArray<FHitResult> Hits;
+	FCollisionShape SphereShape;
+	SphereShape.SetSphere(DesiredInstance.TraceRadius);
+
+	InWorld->SweepMulti(Hits, StartTrace, DesiredInstance.EndTrace, FQuat::Identity, SphereShape, QueryParams, FCollisionObjectQueryParams(ECC_WorldStatic));
+	for (const FHitResult& Hit : Hits)
+	{
+		if (DesiredInstance.PlacementMode == EFoliagePlacementMode::Procedural)
+		{
+			if (Hit.Actor.IsValid())	//if we hit the ProceduralFoliage blocking volume don't spawn instance
+			{
+				if (Cast<AProceduralFoliageBlockingVolume>(Hit.Actor.Get()) || Cast<AInstancedFoliageActor>(Hit.Actor.Get()))
+				{
+					return nullptr;
+				}
+			}
+		}
+			
+		// In the editor traces can hit "No Collision" type actors, so ugh.
+		FBodyInstance* BodyInstance = Hit.Component->GetBodyInstance();
+		if (BodyInstance->GetCollisionEnabled() != ECollisionEnabled::QueryAndPhysics || BodyInstance->GetResponseToChannel(ECC_WorldStatic) != ECR_Block)
+		{
+			continue;
+		}
+
+		if (Hit.Component.IsValid())
+		{
+			if (ULevel* OwningLevel = Hit.Component->GetComponentLevel())
+			{
+				OutHit = Hit;
+				return AInstancedFoliageActor::GetInstancedFoliageActorForLevel(OwningLevel);
+			}
+		}
+	}
+
+	return nullptr;
+}
+
+bool AInstancedFoliageActor::CheckCollisionWithWorld(const UFoliageType* Settings, const FFoliageInstance& Inst, const FVector& HitNormal, const FVector& HitLocation)
+{
+	FMatrix InstTransform = Inst.GetInstanceWorldTransform().ToMatrixWithScale();
+	FVector LocalHit = InstTransform.InverseTransformPosition(HitLocation);
+
+	UWorld* World = GetWorld();
+
+	if (Settings->CollisionWithWorld)
+	{
+		// Check for overhanging ledge
+		{
+			FVector LocalSamplePos[4] = { FVector(Settings->LowBoundOriginRadius.Z, 0, 0),
+			FVector(-Settings->LowBoundOriginRadius.Z, 0, 0),
+			FVector(0, Settings->LowBoundOriginRadius.Z, 0),
+			FVector(0, -Settings->LowBoundOriginRadius.Z, 0)
+			};
+
+
+			for (uint32 i = 0; i < 4; ++i)
+			{
+				FHitResult Hit;
+				FVector SamplePos = InstTransform.TransformPosition(FVector(Settings->LowBoundOriginRadius.X, Settings->LowBoundOriginRadius.Y, 2.f) + LocalSamplePos[i]);
+				float WorldRadius = (Settings->LowBoundOriginRadius.Z + 2.f)*FMath::Max(Inst.DrawScale3D.X, Inst.DrawScale3D.Y);
+				FVector NormalVector = Settings->AlignToNormal ? HitNormal : FVector(0, 0, 1);
+				if (AInstancedFoliageActor::FoliageTrace(World, Hit, FDesiredFoliageInstance(SamplePos, SamplePos - NormalVector*WorldRadius), this))
+				{
+					if (LocalHit.Z - Inst.ZOffset < Settings->LowBoundOriginRadius.Z)
+					{
+						continue;
+					}
+				}
+				return false;
+			}
+		}
+
+		// Check collision with Bounding Box
+		{
+		FBox MeshBox = Settings->MeshBounds.GetBox();
+		MeshBox.Min.Z = FMath::Min(MeshBox.Max.Z, LocalHit.Z + Settings->MeshBounds.BoxExtent.Z * 0.05f);
+		FBoxSphereBounds ShrinkBound(MeshBox);
+		FBoxSphereBounds WorldBound = ShrinkBound.TransformBy(InstTransform);
+		//::DrawDebugBox(World, WorldBound.Origin, WorldBound.BoxExtent, FColor::Red, true, 10.f);
+		static FName NAME_FoliageCollisionWithWorld = FName(TEXT("FoliageCollisionWithWorld"));
+		if (World->OverlapTest(WorldBound.Origin, FQuat(Inst.Rotation), FCollisionShape::MakeBox(ShrinkBound.BoxExtent * Inst.DrawScale3D * Settings->CollisionScale), FCollisionQueryParams(NAME_FoliageCollisionWithWorld, false), FCollisionObjectQueryParams(ECC_WorldStatic)))
+		{
+			return false;
+		}
+	}
+	}
+
+	return true;
+}
+
+FPotentialInstance::FPotentialInstance(FVector InHitLocation, FVector InHitNormal, UPrimitiveComponent* InHitComponent, float InHitWeight, FFoliageMeshInfo* InMeshInfo, AInstancedFoliageActor* InIFA, const FDesiredFoliageInstance& InDesiredInstance)
+: HitLocation(InHitLocation)
+, HitNormal(InHitNormal)
+, HitComponent(InHitComponent)
+, HitWeight(InHitWeight)
+, MeshInfo(InMeshInfo)
+, IFA(InIFA)
+, DesiredInstance(InDesiredInstance)
+{}
+
+bool FPotentialInstance::PlaceInstance(const UFoliageType* Settings, FFoliageInstance& Inst, bool bSkipCollision)
+{
+	if (DesiredInstance.PlacementMode != EFoliagePlacementMode::Procedural)
+	{
+		if (Settings->UniformScale)
+		{
+			float Scale = Settings->ScaleMinX + FMath::FRand() * (Settings->ScaleMaxX - Settings->ScaleMinX);
+			Inst.DrawScale3D = FVector(Scale, Scale, Scale);
+		}
+		else
+		{
+			float LockRand = FMath::FRand();
+			Inst.DrawScale3D.X = Settings->ScaleMinX + (Settings->LockScaleX ? LockRand : FMath::FRand()) * (Settings->ScaleMaxX - Settings->ScaleMinX);
+			Inst.DrawScale3D.Y = Settings->ScaleMinY + (Settings->LockScaleY ? LockRand : FMath::FRand()) * (Settings->ScaleMaxY - Settings->ScaleMinY);
+			Inst.DrawScale3D.Z = Settings->ScaleMinZ + (Settings->LockScaleZ ? LockRand : FMath::FRand()) * (Settings->ScaleMaxZ - Settings->ScaleMinZ);
+		}
+	}
+	else
+	{
+		//Procedural foliage uses age to get the scale
+		Inst.DrawScale3D = FVector(Settings->GetScaleForAge(DesiredInstance.Age));
+	}
+
+
+	Inst.ZOffset = Settings->ZOffsetMin + FMath::FRand() * (Settings->ZOffsetMax - Settings->ZOffsetMin);
+
+	Inst.Location = HitLocation;
+
+	if (DesiredInstance.PlacementMode != EFoliagePlacementMode::Procedural)
+	{
+		// Random yaw and optional random pitch up to the maximum
+		Inst.Rotation = FRotator(FMath::FRand() * Settings->RandomPitchAngle, 0.f, 0.f);
+
+		if (Settings->RandomYaw)
+		{
+			Inst.Rotation.Yaw = FMath::FRand() * 360.f;
+		}
+		else
+		{
+			Inst.Flags |= FOLIAGE_NoRandomYaw;
+		}
+	}
+	else
+	{
+		Inst.Rotation = DesiredInstance.Rotation.Rotator();
+		Inst.Flags |= FOLIAGE_NoRandomYaw;
+	}
+
+
+	if (Settings->AlignToNormal)
+	{
+		Inst.AlignToNormal(HitNormal, Settings->AlignMaxAngle);
+	}
+
+	// Apply the Z offset in local space
+	if (FMath::Abs(Inst.ZOffset) > KINDA_SMALL_NUMBER)
+	{
+		Inst.Location = Inst.GetInstanceWorldTransform().TransformPosition(FVector(0, 0, Inst.ZOffset));
+	}
+
+	UModelComponent* ModelComponent = Cast<UModelComponent>(HitComponent);
+	if (ModelComponent)
+	{
+		ABrush* BrushActor = ModelComponent->GetModel()->FindBrush(HitLocation);
+		if (BrushActor)
+		{
+			HitComponent = BrushActor->GetBrushComponent();
+		}
+	}
+
+	Inst.Base = HitComponent;
+
+	return bSkipCollision || IFA->CheckCollisionWithWorld(Settings, Inst, HitNormal, HitLocation);
+}
+
 #undef LOCTEXT_NAMESPACE
diff --git a/Engine/Source/Runtime/Foliage/Private/ProceduralFoliageComponent.cpp b/Engine/Source/Runtime/Foliage/Private/ProceduralFoliageComponent.cpp
--- a/Engine/Source/Runtime/Foliage/Private/ProceduralFoliageComponent.cpp
+++ b/Engine/Source/Runtime/Foliage/Private/ProceduralFoliageComponent.cpp
@@ -21,55 +21,6 @@ UProceduralFoliageComponent::UProceduralFoliageComponent(const FObjectInitialize
 	ProceduralGuid = FGuid::NewGuid();
 }
 
-void UProceduralFoliageComponent::SpawnInstances(const TArray<FProceduralFoliageInstance>& ProceduralFoliageInstances)
-{
-#if WITH_EDITOR
-	UWorld* World = GetWorld();
-
-	FFoliageInstance Inst;
-	
-	for (const FProceduralFoliageInstance& EcoInst : ProceduralFoliageInstances)
-	{
-		if (EcoInst.BaseComponent)
-		{
-			ULevel* OwningLevel = EcoInst.BaseComponent->GetComponentLevel();
-		
-			AInstancedFoliageActor* IFA = AInstancedFoliageActor::GetInstancedFoliageActorForLevel(OwningLevel);
-			const UFoliageType_InstancedStaticMesh* Type = EcoInst.Type;
-			FFoliageMeshInfo* MeshInfo;
-			if (UStaticMesh* StaticMesh = Type->GetStaticMesh())
-			{
-				if (IFA->GetSettingsForMesh(StaticMesh) == nullptr)
-				{
-					IFA->AddMesh(StaticMesh, nullptr, Type);
-				}
-
-				UFoliageType* Settings = IFA->GetSettingsForMesh(StaticMesh, &MeshInfo);
-				if (Settings && MeshInfo)
-				{
-					/* Convert ProceduralFoliage instance into foliage instance*/
-					Inst.DrawScale3D = FVector(EcoInst.Scale);
-					Inst.Location = EcoInst.Location;
-					Inst.Rotation = EcoInst.Rotation.Rotator();
-					Inst.Flags |= FOLIAGE_NoRandomYaw;
-
-					if (Type->AlignToNormal)
-					{
-						Inst.AlignToNormal(EcoInst.Normal, Type->AlignMaxAngle);
-					}
-
-
-					Inst.Base = EcoInst.BaseComponent;
-					Inst.ProceduralGuid = ProceduralGuid;
-					
-					MeshInfo->AddInstance(IFA, Settings, Inst);
-				}
-			}
-		}
-	}
-#endif
-}
-
 void CopyTileInstances(const UProceduralFoliageTile* FromTile, UProceduralFoliageTile* ToTile, const FBox2D& InnerLocalAABB, const FTransform& ToLocalTM, const float Overlap)
 {
 	const FBox2D OuterLocalAABB(InnerLocalAABB.Min, InnerLocalAABB.Max + FVector2D(Overlap, Overlap));
@@ -97,27 +48,7 @@ FBox2D GetTileRegion(const int32 X, const int32 Y, const int32 CountX, const int
 	return Region;
 }
 
-FBox2D GetRightOverlap(const FBox2D& Region, const float InnerSize, const float Overlap)
-{
-	return Region + FVector2D(InnerSize + Overlap, 0.f);
-}
-
-FBox2D GetLeftOverlap(const FBox2D& Region, const float Overlap)
-{
-	return Region + FVector2D(-Overlap, 0.f);
-}
-
-FBox2D GetAboveOverlap(const FBox2D& Region, const float InnerSize, const float Overlap)
-{
-	return Region + FVector2D(0.f, InnerSize + Overlap);
-}
-
-FBox2D GetBelowOverlap(const FBox2D& Region, const float Overlap)
-{
-	return Region + FVector2D(0.f, -Overlap);
-}
-
-void UProceduralFoliageComponent::SpawnTiles()
+void UProceduralFoliageComponent::SpawnTiles(TArray<FDesiredFoliageInstance>& OutInstances)
 {
 #if WITH_EDITOR
 	if (ProceduralFoliage)
@@ -136,7 +67,7 @@ void UProceduralFoliageComponent::SpawnTiles()
 		const FVector2D OverlapX(Overlap, 0.f);
 		const FVector2D OverlapY(0.f, Overlap);
 
-		TArray<TFuture< TArray<FProceduralFoliageInstance>* >> Futures;
+		TArray<TFuture< TArray<FDesiredFoliageInstance>* >> Futures;
 		FScopedSlowTask SlowTask(TilesX * TilesY, LOCTEXT("PlaceProceduralFoliage", "Placing ProceduralFoliage..."));
 		SlowTask.MakeDialog();
 
@@ -153,7 +84,7 @@ void UProceduralFoliageComponent::SpawnTiles()
 
 				UProceduralFoliageTile* JTile = ProceduralFoliage->CreateTempTile();
 
-				Futures.Add(Async<TArray<FProceduralFoliageInstance>*>(EAsyncExecution::ThreadPool, [=]()
+				Futures.Add(Async<TArray<FDesiredFoliageInstance>*>(EAsyncExecution::ThreadPool, [=]()
 				{
 					FTransform TileTM = ComponentToWorld;
 					const FVector OrientedOffset = ComponentToWorld.TransformVectorNoScale(FVector(X, Y, 0.f) * FVector(InnerTileSize, InnerTileSize, 0.f));
@@ -188,12 +119,12 @@ void UProceduralFoliageComponent::SpawnTiles()
 						CopyTileInstances(TopRightTile, JTile, TopRightBox, TopRightTM, Overlap);
 					}
 
-					TArray<FProceduralFoliageInstance>* ProceduralFoliageInstances = new TArray<FProceduralFoliageInstance>();
+					TArray<FDesiredFoliageInstance>* DesiredInstances = new TArray<FDesiredFoliageInstance>();
 					JTile->InstancesToArray();
-					JTile->CreateInstancesToSpawn(*ProceduralFoliageInstances, TileTM, World, HalfHeight);
+					JTile->CreateInstancesToSpawn(*DesiredInstances, TileTM, ProceduralGuid, HalfHeight);
 					JTile->Empty();
 
-					return ProceduralFoliageInstances;
+					return DesiredInstances;
 					
 				})
 				);
@@ -205,9 +136,9 @@ void UProceduralFoliageComponent::SpawnTiles()
 		{
 			for (int Y = 0; Y < TilesY; ++Y)
 			{
-				TArray<FProceduralFoliageInstance>* ProceduralFoliageInstances = Futures[FutureIdx++].Get();
-				SpawnInstances(*ProceduralFoliageInstances);
-				delete ProceduralFoliageInstances;
+				TArray<FDesiredFoliageInstance>* DesiredInstances = Futures[FutureIdx++].Get();
+				OutInstances.Append(*DesiredInstances);
+				delete DesiredInstances;
 				SlowTask.EnterProgressFrame(1);
 			}
 		}
@@ -215,11 +146,11 @@ void UProceduralFoliageComponent::SpawnTiles()
 #endif
 }
 
-void UProceduralFoliageComponent::SpawnProceduralContent()
+void UProceduralFoliageComponent::SpawnProceduralContent(TArray <FDesiredFoliageInstance>& OutInstances)
 {
 #if WITH_EDITOR
 	RemoveProceduralContent();
-	SpawnTiles();
+	SpawnTiles(OutInstances);
 #endif
 }
 
diff --git a/Engine/Source/Runtime/Foliage/Private/ProceduralFoliageTile.cpp b/Engine/Source/Runtime/Foliage/Private/ProceduralFoliageTile.cpp
--- a/Engine/Source/Runtime/Foliage/Private/ProceduralFoliageTile.cpp
+++ b/Engine/Source/Runtime/Foliage/Private/ProceduralFoliageTile.cpp
@@ -4,7 +4,6 @@
 #include "ProceduralFoliageTile.h"
 #include "ProceduralFoliage.h"
 #include "ProceduralFoliageBroadphase.h"
-#include "ProceduralFoliageBlockingVolume.h"
 #include "InstancedFoliageActor.h"
 
 #define LOCTEXT_NAMESPACE "ProceduralFoliage"
@@ -334,7 +333,7 @@ void UProceduralFoliageTile::BeginDestroy()
 	RemoveInstances();
 }
 
-void UProceduralFoliageTile::CreateInstancesToSpawn(TArray<FProceduralFoliageInstance>& OutInstances, const FTransform& WorldTM, UWorld* World, const float HalfHeight) const
+void UProceduralFoliageTile::CreateInstancesToSpawn(TArray<FDesiredFoliageInstance>& OutInstances, const FTransform& WorldTM, const FGuid& ProceduralGuid, const float HalfHeight) const
 {
 	const FCollisionQueryParams Params(true);
 	FHitResult Hit;
@@ -342,41 +341,17 @@ void UProceduralFoliageTile::CreateInstancesToSpawn(TArray<FProceduralFoliageIns
 	OutInstances.Reserve(Instances.Num());
 	for (const FProceduralFoliageInstance& Instance : InstancesArray)
 	{
-		//@todo ProceduralFoliage need a better method for calculating the ray
 		FVector StartRay = Instance.Location + WorldTM.GetLocation();
 		StartRay.Z += HalfHeight;
 		FVector EndRay = StartRay;
 		EndRay.Z -= HalfHeight*2.f;
-		FCollisionShape SphereShape;
-		SphereShape.SetSphere(Instance.GetMaxRadius());
 
-		if (World->SweepSingle(Hit, StartRay, EndRay, FQuat::Identity, SphereShape, Params, FCollisionObjectQueryParams(ECC_WorldStatic)))
-		{
-			if (Hit.Actor.IsValid())	//if we hit the ProceduralFoliage blocking volume don't spawn instance
-			{
-				if (Cast<AProceduralFoliageBlockingVolume>(Hit.Actor.Get()) || Cast<AInstancedFoliageActor>(Hit.Actor.Get()))
-				{
-					continue;
-				}
-			}
-
-			const UFoliageType_InstancedStaticMesh* Type = Instance.Type;
-			if (Hit.ImpactPoint.Z >= Type->HeightMin && Hit.ImpactPoint.Z <= Type->HeightMax)
-			{
-				const float MaxNormalAngle = FMath::Cos(FMath::DegreesToRadians(Type->GroundSlope));
-				const float MinNormalAngle = FMath::Cos(FMath::DegreesToRadians(Type->MinGroundSlope));
-				if (MaxNormalAngle <= Hit.ImpactNormal.Z && MinNormalAngle >= Hit.ImpactNormal.Z)	//keep in mind Hit.ImpactNormal.Z is (0,0,1) dot normal. However, ground slope is with relation to plane vector, not plane normal - so we swap comparisons
-				{
-					FProceduralFoliageInstance* NewInst = new(OutInstances)FProceduralFoliageInstance(Instance);
-					NewInst->Location = FVector(StartRay.X, StartRay.Y, Hit.ImpactPoint.Z);	//take the x,y of the instance, but use the z of the impact point. This is because we never want to move the instance along xy or we'll get overlaps 
-					NewInst->Normal = Hit.ImpactNormal;
-					if (Hit.Component.IsValid())
-					{
-						NewInst->BaseComponent = Hit.Component.Get();
-					}
-				}
-			}
-		}
+		FDesiredFoliageInstance* DesiredInst = new (OutInstances)FDesiredFoliageInstance(StartRay, EndRay, Instance.GetMaxRadius());
+		DesiredInst->Rotation = Instance.Rotation;
+		DesiredInst->ProceduralGuid = ProceduralGuid;
+		DesiredInst->FoliageType = Instance.Type;
+		DesiredInst->Age = Instance.Age;
+		DesiredInst->PlacementMode = EFoliagePlacementMode::Procedural;
 	}
 }
 
diff --git a/Engine/Source/Runtime/Foliage/Public/FoliageType.h b/Engine/Source/Runtime/Foliage/Public/FoliageType.h
--- a/Engine/Source/Runtime/Foliage/Public/FoliageType.h
+++ b/Engine/Source/Runtime/Foliage/Public/FoliageType.h
@@ -79,6 +79,9 @@ class UFoliageType : public UObject
 	UPROPERTY(EditAnywhere, Category = General, meta = (Subcategory = "Placement"))
 	FName LandscapeLayer;
 
+	UPROPERTY(EditAnywhere, Category = General, meta = (Subcategory = "Placement"))
+	float MinimumLayerWeight;
+
 	UPROPERTY(EditAnywhere, Category = General, meta = (Subcategory = "Placement"))
 	uint32 AlignToNormal:1;
 
diff --git a/Engine/Source/Runtime/Foliage/Public/InstancedFoliage.h b/Engine/Source/Runtime/Foliage/Public/InstancedFoliage.h
--- a/Engine/Source/Runtime/Foliage/Public/InstancedFoliage.h
+++ b/Engine/Source/Runtime/Foliage/Public/InstancedFoliage.h
@@ -207,7 +207,7 @@ struct FFoliageMeshInfo
 	}
 
 #if WITH_EDITOR
-	FOLIAGE_API void AddInstance(AInstancedFoliageActor* InIFA, UFoliageType* InSettings, const FFoliageInstance& InNewInstance);
+	FOLIAGE_API void AddInstance(AInstancedFoliageActor* InIFA, const UFoliageType* InSettings, const FFoliageInstance& InNewInstance);
 	FOLIAGE_API void RemoveInstances(AInstancedFoliageActor* InIFA, const TArray<int32>& InInstancesToRemove);
 	FOLIAGE_API void PreMoveInstances(AInstancedFoliageActor* InIFA, const TArray<int32>& InInstancesToMove);
 	FOLIAGE_API void PostMoveInstances(AInstancedFoliageActor* InIFA, const TArray<int32>& InInstancesMoved);
@@ -347,3 +347,69 @@ public:
 		return Ar;
 	}
 };
+
+
+/** This is kind of a hack, but is needed right now for backwards compat of code. We use it to describe the placement mode (procedural vs manual)*/
+namespace EFoliagePlacementMode
+{
+	enum Type
+	{
+		Manual = 0,
+		Procedural = 1,
+	};
+
+}
+
+/** Used to define a vector along which we'd like to spawn an instance. */
+struct FDesiredFoliageInstance
+{
+	FDesiredFoliageInstance()
+	: FoliageType(nullptr)
+	, MeshInfo(nullptr)
+	, StartTrace(ForceInit)
+	, EndTrace(ForceInit)
+	, Rotation(ForceInit)
+	, TraceRadius(0.f)
+	, Age(0.f)
+	, PlacementMode(EFoliagePlacementMode::Manual)
+	{
+
+	}
+
+	FDesiredFoliageInstance(const FVector& InStartTrace, const FVector& InEndTrace, const float InTraceRadius = 0.f)
+	: FoliageType(nullptr)
+	, MeshInfo(nullptr)
+	, StartTrace(InStartTrace)
+	, EndTrace(InEndTrace)
+	, Rotation(ForceInit)
+	, TraceRadius(InTraceRadius)
+	, Age(0.f)
+	, PlacementMode(EFoliagePlacementMode::Manual)
+	{
+	}
+
+	FGuid ProceduralGuid;
+	const UFoliageType* FoliageType;
+	FFoliageMeshInfo* MeshInfo;
+	FVector StartTrace;
+	FVector EndTrace;
+	FQuat Rotation;
+	float TraceRadius;
+	float Age;
+	EFoliagePlacementMode::Type PlacementMode;
+};
+
+// Struct to hold potential instances we've sampled
+struct FOLIAGE_API FPotentialInstance
+{
+	FVector HitLocation;
+	FVector HitNormal;
+	UPrimitiveComponent* HitComponent;
+	float HitWeight;
+	FFoliageMeshInfo* MeshInfo;
+	AInstancedFoliageActor* IFA;
+	FDesiredFoliageInstance DesiredInstance;
+
+	FPotentialInstance(FVector InHitLocation, FVector InHitNormal, UPrimitiveComponent* InHitComponent, float InHitWeight, FFoliageMeshInfo* MeshInfo = nullptr, AInstancedFoliageActor* InIFA = nullptr, const FDesiredFoliageInstance& InDesiredInstance = FDesiredFoliageInstance());
+	bool PlaceInstance(const UFoliageType* Settings, FFoliageInstance& Inst, bool bSkipCollision = false);
+};
diff --git a/Engine/Source/Runtime/Foliage/Public/InstancedFoliageActor.h b/Engine/Source/Runtime/Foliage/Public/InstancedFoliageActor.h
--- a/Engine/Source/Runtime/Foliage/Public/InstancedFoliageActor.h
+++ b/Engine/Source/Runtime/Foliage/Public/InstancedFoliageActor.h
@@ -15,6 +15,9 @@ struct FFoliageInstancePlacementInfo;
 struct FFoliageMeshInfo;
 class UFoliageType_InstancedStaticMesh;
 class UProceduralFoliageComponent;
+struct FFoliageInstance;
+struct FHitResult;
+struct FDesiredFoliageInstance;
 
 UCLASS(notplaceable, hidecategories = Object, MinimalAPI, NotBlueprintable)
 class AInstancedFoliageActor : public AActor
@@ -73,6 +76,9 @@ public:
 	*/
 	static FOLIAGE_API AInstancedFoliageActor* GetInstancedFoliageActorForLevel(ULevel* Level, bool bCreateIfNone = true);
 
+	static FOLIAGE_API AInstancedFoliageActor* FoliageTrace(UWorld* InWorld, FHitResult& OutHit, const FDesiredFoliageInstance& DesiredInstance, const AInstancedFoliageActor* IgnoreIFA = nullptr, FName InTraceTag = NAME_None, bool InbReturnFaceIndex = false);
+	FOLIAGE_API bool CheckCollisionWithWorld(const UFoliageType* Settings, const FFoliageInstance& Inst, const FVector& HitNormal, const FVector& HitLocation);
+
 #if WITH_EDITOR
 	virtual void PostEditUndo() override;
 
diff --git a/Engine/Source/Runtime/Foliage/Public/ProceduralFoliageComponent.h b/Engine/Source/Runtime/Foliage/Public/ProceduralFoliageComponent.h
--- a/Engine/Source/Runtime/Foliage/Public/ProceduralFoliageComponent.h
+++ b/Engine/Source/Runtime/Foliage/Public/ProceduralFoliageComponent.h
@@ -7,6 +7,7 @@
 
 class UProceduralFoliage;
 class AProceduralFoliageLevelInfo;
+struct FDesiredFoliageInstance;
 
 UCLASS(BlueprintType)
 class FOLIAGE_API UProceduralFoliageComponent : public USceneComponent
@@ -32,13 +33,12 @@ class FOLIAGE_API UProceduralFoliageComponent : public USceneComponent
 	UPROPERTY(Category = "ProceduralFoliage", BlueprintReadWrite, EditAnywhere, meta=(DisplayName="ProceduralFoliage Asset") )
 	UProceduralFoliage* ProceduralFoliage;
 
-	void SpawnProceduralContent();
+	void SpawnProceduralContent(TArray<FDesiredFoliageInstance>& OutFoliageInstances);
 	void RemoveProceduralContent();
 	const FGuid& GetProceduralGuid() const { return ProceduralGuid;  }
 
 private:
-	void SpawnTiles();
-	void SpawnInstances(const TArray<FProceduralFoliageInstance>& ProceduralFoliageInstances);
+	void SpawnTiles(TArray<FDesiredFoliageInstance>& OutFoliageInstances);
 
 private:
 	UPROPERTY()
diff --git a/Engine/Source/Runtime/Foliage/Public/ProceduralFoliageTile.h b/Engine/Source/Runtime/Foliage/Public/ProceduralFoliageTile.h
--- a/Engine/Source/Runtime/Foliage/Public/ProceduralFoliageTile.h
+++ b/Engine/Source/Runtime/Foliage/Public/ProceduralFoliageTile.h
@@ -4,6 +4,7 @@
 
 #include "ProceduralFoliageBroadphase.h"
 #include "ProceduralFoliageInstance.h"
+#include "InstancedFoliage.h"
 #include "ProceduralFoliageTile.generated.h"
 
 struct FProceduralFoliageInstance;
@@ -26,7 +27,7 @@ class FOLIAGE_API UProceduralFoliageTile : public UObject
 
 	/** Takes a local AABB and finds all the instances in it. The instances are returned in the space local to the AABB. That is, an instance in the bottom left AABB will have a coordinate of (0,0) */
 	void GetInstancesInAABB(const FBox2D& LocalAABB, TArray<FProceduralFoliageInstance* >& OutInstances, bool bOnTheBorder) const;
-	void CreateInstancesToSpawn(TArray<FProceduralFoliageInstance>& OutInstances, const FTransform& WorldTM, UWorld* World, const float HalfHeight) const;
+	void CreateInstancesToSpawn(TArray<FDesiredFoliageInstance>& OutInstances, const FTransform& WorldTM, const FGuid& ProceduralGuid, const float HalfHeight) const;
 	void AddInstances(const TArray<FProceduralFoliageInstance*>& NewInstances, const FTransform& LocalTM, const FBox2D& InnerLocalAABB);
 
 	virtual SIZE_T GetResourceSize(EResourceSizeMode::Type Mode) override;
