Forgot to pass UISettings pointer into static function

diff --git a/Engine/Source/Editor/FoliageEdit/Private/FoliageEdMode.cpp b/Engine/Source/Editor/FoliageEdit/Private/FoliageEdMode.cpp
--- a/Engine/Source/Editor/FoliageEdit/Private/FoliageEdMode.cpp
+++ b/Engine/Source/Editor/FoliageEdit/Private/FoliageEdMode.cpp
@@ -436,7 +436,7 @@ static bool CheckVertexColor(const UFoliageType* Settings, const FColor& VertexC
 	return true;
 }
 
-void FEdModeFoliage::CalculatePotentialInstances(UWorld* InWorld, const AInstancedFoliageActor* IgnoreIFA, const UFoliageType* FoliageType, const TArray<FDesiredFoliageInstance>& DesiredInstances, TArray<FPotentialInstance> OutPotentialInstances[NUM_INSTANCE_BUCKETS], LandscapeLayerCacheData* LandscapeLayerCachesPtr)
+void FEdModeFoliage::CalculatePotentialInstances(UWorld* InWorld, const AInstancedFoliageActor* IgnoreIFA, const UFoliageType* FoliageType, const TArray<FDesiredFoliageInstance>& DesiredInstances, TArray<FPotentialInstance> OutPotentialInstances[NUM_INSTANCE_BUCKETS], LandscapeLayerCacheData* LandscapeLayerCachesPtr, const FFoliageUISettings* UISettings)
 {
 	LandscapeLayerCacheData LocalCache;
 	LandscapeLayerCachesPtr = LandscapeLayerCachesPtr ? LandscapeLayerCachesPtr : &LocalCache;
@@ -462,30 +462,27 @@ void FEdModeFoliage::CalculatePotentialInstances(UWorld* InWorld, const AInstanc
 		if (AInstancedFoliageActor* IFA = AInstancedFoliageActor::FoliageTrace(InWorld, Hit, DesiredInst, IgnoreIFA, NAME_AddFoliageInstances, true))
 		{
 			// Check filters
-			UPrimitiveComponent* PrimComp = Hit.Component.Get();
-			UMaterialInterface* Material = PrimComp ? PrimComp->GetMaterial(0) : nullptr;
 
-			if (DesiredInst.PlacementMode == EFoliagePlacementMode::Manual)
+			if (UPrimitiveComponent* PrimComp = Hit.Component.Get())
 			{
-				if (!PrimComp ||
-					PrimComp->GetOutermost() != InWorld->GetCurrentLevel()->GetOutermost()/* ||
-					(!UISettings.bFilterLandscape && PrimComp->IsA(ULandscapeHeightfieldCollisionComponent::StaticClass())) ||
-					(!UISettings.bFilterStaticMesh && PrimComp->IsA(UStaticMeshComponent::StaticClass())) ||
-					(!UISettings.bFilterBSP && PrimComp->IsA(UModelComponent::StaticClass())) ||
-					(!UISettings.bFilterTranslucent && Material && IsTranslucentBlendMode(Material->GetBlendMode()))*/
-					)
+				if (DesiredInst.PlacementMode == EFoliagePlacementMode::Manual)
 				{
-					continue;
+					UMaterialInterface* Material = PrimComp ? PrimComp->GetMaterial(0) : nullptr;
+					if (PrimComp->GetOutermost() != InWorld->GetCurrentLevel()->GetOutermost() ||
+						(!UISettings->bFilterLandscape && PrimComp->IsA(ULandscapeHeightfieldCollisionComponent::StaticClass())) ||
+						(!UISettings->bFilterStaticMesh && PrimComp->IsA(UStaticMeshComponent::StaticClass())) ||
+						(!UISettings->bFilterBSP && PrimComp->IsA(UModelComponent::StaticClass())) ||
+						(!UISettings->bFilterTranslucent && Material && IsTranslucentBlendMode(Material->GetBlendMode()))
+						)
+					{
+						continue;
+					}
 				}
 			}
 			else
 			{
-				if (PrimComp == false)
-				{
-					continue;
-				}
+				continue;
 			}
-			
 
 
 			FFoliageMeshInfo* MeshInfo = DesiredInst.MeshInfo;
@@ -591,10 +588,10 @@ void FEdModeFoliage::AddInstances(UWorld* InWorld, const TArray<FDesiredFoliageI
 	}
 }
 
-void FEdModeFoliage::AddInstancesImp(UWorld* InWorld, const UFoliageType* Settings, const TArray<FDesiredFoliageInstance>& DesiredInstances, const AInstancedFoliageActor* IgnoreIFA, const TArray<int32>& ExistingInstanceBuckets, const float Pressure, LandscapeLayerCacheData* LandscapeLayerCachesPtr)
+void FEdModeFoliage::AddInstancesImp(UWorld* InWorld, const UFoliageType* Settings, const TArray<FDesiredFoliageInstance>& DesiredInstances, const AInstancedFoliageActor* IgnoreIFA, const TArray<int32>& ExistingInstanceBuckets, const float Pressure, LandscapeLayerCacheData* LandscapeLayerCachesPtr, const FFoliageUISettings* UISettings)
 {
 	TArray<FPotentialInstance> PotentialInstanceBuckets[NUM_INSTANCE_BUCKETS];
-	CalculatePotentialInstances(InWorld, IgnoreIFA, Settings, DesiredInstances, PotentialInstanceBuckets, LandscapeLayerCachesPtr);
+	CalculatePotentialInstances(InWorld, IgnoreIFA, Settings, DesiredInstances, PotentialInstanceBuckets, LandscapeLayerCachesPtr, UISettings);
 
 	for (int32 BucketIdx = 0; BucketIdx < NUM_INSTANCE_BUCKETS; BucketIdx++)
 	{
@@ -673,7 +670,7 @@ void FEdModeFoliage::AddInstancesForBrush(UWorld* InWorld, AInstancedFoliageActo
 			DesiredInstance->MeshInfo = &MeshInfo;
 		}
 
-		AddInstancesImp(InWorld, Settings, DesiredInstances, IFA, ExistingInstanceBuckets, Pressure, &LandscapeLayerCaches);
+		AddInstancesImp(InWorld, Settings, DesiredInstances, IFA, ExistingInstanceBuckets, Pressure, &LandscapeLayerCaches, &UISettings);
 	}
 }
 
diff --git a/Engine/Source/Editor/FoliageEdit/Private/FoliageEdMode.h b/Engine/Source/Editor/FoliageEdit/Private/FoliageEdMode.h
--- a/Engine/Source/Editor/FoliageEdit/Private/FoliageEdMode.h
+++ b/Engine/Source/Editor/FoliageEdit/Private/FoliageEdMode.h
@@ -315,9 +315,9 @@ private:
 	void AddInstancesForBrush(UWorld* InWorld, AInstancedFoliageActor* IFA, UFoliageType* Settings, FFoliageMeshInfo& MeshInfo, int32 DesiredInstanceCount, const TArray<int32>& ExistingInstances, float Pressure);
 
 	/** Common code for adding instances to world based on settings */
-	static void AddInstancesImp(UWorld* InWorld, const UFoliageType* Settings, const TArray<FDesiredFoliageInstance>& DesiredInstances, const AInstancedFoliageActor* IgnoreIFA = nullptr, const TArray<int32>& ExistingInstances = TArray<int32>(), const float Pressure = 1.f, LandscapeLayerCacheData* LandscapeLayerCaches = nullptr);
+	static void AddInstancesImp(UWorld* InWorld, const UFoliageType* Settings, const TArray<FDesiredFoliageInstance>& DesiredInstances, const AInstancedFoliageActor* IgnoreIFA = nullptr, const TArray<int32>& ExistingInstances = TArray<int32>(), const float Pressure = 1.f, LandscapeLayerCacheData* LandscapeLayerCaches = nullptr, const FFoliageUISettings* UISettings = nullptr);
 
-	static void CalculatePotentialInstances(UWorld* InWorld, const AInstancedFoliageActor* IFA, const UFoliageType* Settings, const TArray<FDesiredFoliageInstance>& DesiredInstances, TArray<FPotentialInstance> OutPotentialInstances[NUM_INSTANCE_BUCKETS], LandscapeLayerCacheData* LandscaleLayerCachesPtr);
+	static void CalculatePotentialInstances(UWorld* InWorld, const AInstancedFoliageActor* IFA, const UFoliageType* Settings, const TArray<FDesiredFoliageInstance>& DesiredInstances, TArray<FPotentialInstance> OutPotentialInstances[NUM_INSTANCE_BUCKETS], LandscapeLayerCacheData* LandscaleLayerCachesPtr, const FFoliageUISettings* UISettings);
 
 	/** Remove instances inside the brush to match DesiredInstanceCount. NOTE: PotentialInstancesToRemove array is modified by this function. */
 	void RemoveInstancesForBrush(AInstancedFoliageActor* IFA, FFoliageMeshInfo& MeshInfo, int32 DesiredInstanceCount, TArray<int32>& PotentialInstancesToRemove, float Pressure);
