Replace AddUninitialized with AddZeroed. Expecting 0 in the array elements

diff --git a/Engine/Source/Editor/FoliageEdit/Private/FoliageEdMode.cpp b/Engine/Source/Editor/FoliageEdit/Private/FoliageEdMode.cpp
--- a/Engine/Source/Editor/FoliageEdit/Private/FoliageEdMode.cpp
+++ b/Engine/Source/Editor/FoliageEdit/Private/FoliageEdMode.cpp
@@ -622,7 +622,7 @@ void FEdModeFoliage::AddInstancesForBrush(UWorld* InWorld, AInstancedFoliageActo
 	if (DesiredInstanceCount > ExistingInstances.Num())
 	{
 		TArray<int32> ExistingInstanceBuckets;
-		ExistingInstanceBuckets.AddUninitialized(NUM_INSTANCE_BUCKETS);
+		ExistingInstanceBuckets.AddZeroed(NUM_INSTANCE_BUCKETS);
 
 		// Cache store mapping between component and weight data
 		TMap<ULandscapeComponent*, TArray<uint8> >* LandscapeLayerCache = nullptr;
