Add multiple layer support for foliage painting

diff --git a/Engine/Config/BaseEngine.ini b/Engine/Config/BaseEngine.ini
--- a/Engine/Config/BaseEngine.ini
+++ b/Engine/Config/BaseEngine.ini
@@ -1137,6 +1137,7 @@ bUseStreamingPause=false
 +ActiveClassRedirects=(OldClassName="InstancedFoliageActor", NewClassName="/Script/Foliage.InstancedFoliageActor")
 +ActiveClassRedirects=(OldClassName="InteractiveFoliageComponent", NewClassName="/Script/Foliage.InteractiveFoliageComponent")
 +ActiveClassRedirects=(OldClassName="FoliageVertexColorMask", NewClassName="/Script/Foliage.FoliageVertexColorMask")
++TaggedPropertyRedirects(ClassName="FoliageType", OldPropertyName="LandscapeLayer", NewPropertyName="LandscapeLayer_DEPRECATED")
 
 
 
diff --git a/Engine/Source/Editor/FoliageEdit/Private/FoliageEdMode.cpp b/Engine/Source/Editor/FoliageEdit/Private/FoliageEdMode.cpp
--- a/Engine/Source/Editor/FoliageEdit/Private/FoliageEdMode.cpp
+++ b/Engine/Source/Editor/FoliageEdit/Private/FoliageEdMode.cpp
@@ -436,13 +436,53 @@ static bool CheckVertexColor(const UFoliageType* Settings, const FColor& VertexC
 	return true;
 }
 
-void FEdModeFoliage::CalculatePotentialInstances(UWorld* InWorld, const AInstancedFoliageActor* IgnoreIFA, const UFoliageType* FoliageType, const TArray<FDesiredFoliageInstance>& DesiredInstances, TArray<FPotentialInstance> OutPotentialInstances[NUM_INSTANCE_BUCKETS], LandscapeLayerCacheData* LandscapeLayerCachesPtr, const FFoliageUISettings* UISettings)
+bool LandscapeLayersValid(const UFoliageType* Settings)
+{
+	bool bValid = false;
+	for (FName LayerName : Settings->LandscapeLayers)
+	{
+		bValid |= LayerName != NAME_None;
+	}
+
+	return bValid;
+}
+
+bool GetMaxHitWeight(const FVector& Location, UActorComponent* Component, const UFoliageType* Settings, FEdModeFoliage::LandscapeLayerCacheData* LandscapeLayerCaches, float& OutMaxHitWeight)
+{
+	float MaxHitWeight = 0.f;
+	if (ULandscapeHeightfieldCollisionComponent* HitLandscapeCollision = Cast<ULandscapeHeightfieldCollisionComponent>(Component))
+	{
+		if (ULandscapeComponent* HitLandscape = HitLandscapeCollision->RenderComponent.Get())
+		{
+			for (FName LandscapeLayerName : Settings->LandscapeLayers)
+			{
+				// Cache store mapping between component and weight data
+				TMap<ULandscapeComponent*, TArray<uint8> >* LandscapeLayerCache = &LandscapeLayerCaches->FindOrAdd(LandscapeLayerName);;
+				TArray<uint8>* LayerCache = &LandscapeLayerCache->FindOrAdd(HitLandscape);
+				// TODO: FName to LayerInfo?
+				const float HitWeight = HitLandscape->GetLayerWeightAtLocation(Location, HitLandscape->GetLandscapeInfo()->GetLayerInfoByName(LandscapeLayerName), LayerCache);
+				MaxHitWeight = FMath::Max(MaxHitWeight, HitWeight);
+			}
+
+			OutMaxHitWeight = MaxHitWeight;
+			return true;
+		}
+	}
+
+	return false;
+}
+
+bool FilterByWeight(float Weight, const UFoliageType* Settings)
+{
+	const float WeightNeeded = FMath::Max(Settings->MinimumLayerWeight, FMath::FRand());
+	return Weight < FMath::Max(SMALL_NUMBER, WeightNeeded);
+}
+
+void FEdModeFoliage::CalculatePotentialInstances(UWorld* InWorld, const AInstancedFoliageActor* IgnoreIFA, const UFoliageType* Settings, const TArray<FDesiredFoliageInstance>& DesiredInstances, TArray<FPotentialInstance> OutPotentialInstances[NUM_INSTANCE_BUCKETS], LandscapeLayerCacheData* LandscapeLayerCachesPtr, const FFoliageUISettings* UISettings)
 {
 	LandscapeLayerCacheData LocalCache;
 	LandscapeLayerCachesPtr = LandscapeLayerCachesPtr ? LandscapeLayerCachesPtr : &LocalCache;
 
-	const UFoliageType* Settings = FoliageType;
-
 	// Quick lookup of potential instance locations, used for overlapping check.
 	TArray<FVector> PotentialInstanceLocations;
 	FFoliageInstanceHash PotentialInstanceHash(7);	// use 128x128 cell size, things like brush radius are typically small
@@ -455,6 +495,9 @@ void FEdModeFoliage::CalculatePotentialInstances(UWorld* InWorld, const AInstanc
 		Bucket.Reserve(DesiredInstances.Num());
 	}
 
+	// Radius where we expect to have a single instance, given the density rules
+	const float DensityCheckRadius = FMath::Max<float>(FMath::Sqrt((1000.f*1000.f) / (PI * Settings->Density)), Settings->Radius);
+
 	for (const FDesiredFoliageInstance& DesiredInst : DesiredInstances)
 	{
 		FHitResult Hit;
@@ -490,7 +533,7 @@ void FEdModeFoliage::CalculatePotentialInstances(UWorld* InWorld, const AInstanc
 			{
 				//We have to add the mesh into the foliage actor. This assumes the foliage type is just a concrete InstancedStaticMesh.
 				//With the new tree system coming this will not be true and you'll need to traverse the tree and insert the needed meshes
-				if (UStaticMesh* StaticMesh = FoliageType->GetStaticMesh())
+				if (UStaticMesh* StaticMesh = Settings->GetStaticMesh())
 				{
 					UFoliageType* ExistingSettings = IFA->GetSettingsForMesh(StaticMesh);
 					if (ExistingSettings == nullptr)
@@ -503,7 +546,6 @@ void FEdModeFoliage::CalculatePotentialInstances(UWorld* InWorld, const AInstanc
 					else
 					{
 						MeshInfo = IFA->FindMesh(ExistingSettings);
-						//Settings = ExistingSettings;
 					}
 				}
 
@@ -515,12 +557,6 @@ void FEdModeFoliage::CalculatePotentialInstances(UWorld* InWorld, const AInstanc
 				continue;
 			}
 
-			// Radius where we expect to have a single instance, given the density rules
-			const float DensityCheckRadius = FMath::Max<float>(FMath::Sqrt((1000.f*1000.f) / (PI * Settings->Density)), Settings->Radius);
-			const FName LandscapeLayerName = Settings->LandscapeLayer;
-			TMap<ULandscapeComponent*, TArray<uint8>>* LandscapeLayerCache = LandscapeLayerName != NAME_None ? &LandscapeLayerCachesPtr->FindOrAdd(LandscapeLayerName) : nullptr;
-			
-
 			if (CheckLocationForPotentialInstance(*MeshInfo, Settings, DensityCheckRadius, Hit.ImpactPoint, Hit.ImpactNormal, PotentialInstanceLocations, PotentialInstanceHash) == false)
 			{
 				continue;
@@ -544,23 +580,12 @@ void FEdModeFoliage::CalculatePotentialInstances(UWorld* InWorld, const AInstanc
 
 			// Check landscape layer
 			float HitWeight = 1.f;
-			if (LandscapeLayerCache)
+			if (LandscapeLayersValid(Settings) && GetMaxHitWeight(Hit.ImpactPoint, Hit.Component.Get(), Settings, LandscapeLayerCachesPtr, HitWeight))
 			{
-				if (ULandscapeHeightfieldCollisionComponent* HitLandscapeCollision = Cast<ULandscapeHeightfieldCollisionComponent>(Hit.Component.Get()))
+				// Reject instance randomly in proportion to weight
+				if (FilterByWeight(HitWeight, Settings))
 				{
-					if (ULandscapeComponent* HitLandscape = HitLandscapeCollision->RenderComponent.Get())
-					{
-						TArray<uint8>* LayerCache = &LandscapeLayerCache->FindOrAdd(HitLandscape);
-						// TODO: FName to LayerInfo?
-						HitWeight = HitLandscape->GetLayerWeightAtLocation(Hit.ImpactPoint, HitLandscape->GetLandscapeInfo()->GetLayerInfoByName(LandscapeLayerName), LayerCache);
-
-						// Reject instance randomly in proportion to weight
-						const float WeightNeeded = FMath::Max(Settings->MinimumLayerWeight, FMath::FRand());
-						if (HitWeight < FMath::Max(SMALL_NUMBER, WeightNeeded))
-						{
-							continue;
-						}
-					}
+					continue;
 				}
 			}
 
@@ -624,31 +649,18 @@ void FEdModeFoliage::AddInstancesForBrush(UWorld* InWorld, AInstancedFoliageActo
 		TArray<int32> ExistingInstanceBuckets;
 		ExistingInstanceBuckets.AddZeroed(NUM_INSTANCE_BUCKETS);
 
-		// Cache store mapping between component and weight data
-		TMap<ULandscapeComponent*, TArray<uint8> >* LandscapeLayerCache = nullptr;
-
-		FName LandscapeLayerName = Settings->LandscapeLayer;
-		if (LandscapeLayerName != NAME_None)
+		const TArray<FName>& LandscapeLayers = Settings->LandscapeLayers;
+		if (LandscapeLayersValid(Settings))
 		{
-			LandscapeLayerCache = &LandscapeLayerCaches.FindOrAdd(LandscapeLayerName);
-
 			// Find the landscape weights of existing ExistingInstances
 			for (int32 Idx : ExistingInstances)
 			{
 				FFoliageInstance& Instance = MeshInfo.Instances[Idx];
-				ULandscapeHeightfieldCollisionComponent* HitLandscapeCollision = Cast<ULandscapeHeightfieldCollisionComponent>(Instance.Base);
-				if (HitLandscapeCollision)
+				float HitWeight;
+				if (GetMaxHitWeight(Instance.Location, Instance.Base, Settings, &LandscapeLayerCaches, HitWeight))
 				{
-					ULandscapeComponent* HitLandscape = HitLandscapeCollision->RenderComponent.Get();
-					if (HitLandscape)
-					{
-						TArray<uint8>* LayerCache = &LandscapeLayerCache->FindOrAdd(HitLandscape);
-						// TODO: FName to LayerInfo?
-						float HitWeight = HitLandscape->GetLayerWeightAtLocation(Instance.Location, HitLandscape->GetLandscapeInfo()->GetLayerInfoByName(LandscapeLayerName), LayerCache);
-
-						// Add count to bucket.
-						ExistingInstanceBuckets[FMath::RoundToInt(HitWeight * (float)(NUM_INSTANCE_BUCKETS - 1))]++;
-					}
+					// Add count to bucket.
+					ExistingInstanceBuckets[FMath::RoundToInt(HitWeight * (float)(NUM_INSTANCE_BUCKETS - 1))]++;
 				}
 			}
 		}
@@ -730,14 +742,6 @@ void FEdModeFoliage::ReapplyInstancesForBrush(UWorld* InWorld, AInstancedFoliage
 	TArray<int32> UpdatedInstances;
 	TSet<int32> InstancesToDelete;
 
-	// Setup cache if we're reapplying landscape layer weights
-	FName LandscapeLayerName = Settings->LandscapeLayer;
-	TMap<ULandscapeComponent*, TArray<uint8>>* LandscapeLayerCache = nullptr;
-	if (Settings->ReapplyLandscapeLayer && LandscapeLayerName != NAME_None)
-	{
-		LandscapeLayerCache = &LandscapeLayerCaches.FindOrAdd(LandscapeLayerName);
-	}
-
 	IFA->Modify();
 
 	for (int32 Idx = 0; Idx < ExistingInstances.Num(); Idx++)
@@ -895,7 +899,7 @@ void FEdModeFoliage::ReapplyInstancesForBrush(UWorld* InWorld, AInstancedFoliage
 			}
 
 			// Find a ground normal for either normal or ground slope check.
-			if (bReapplyNormal || Settings->ReapplyGroundSlope || Settings->ReapplyVertexColorMask || (Settings->ReapplyLandscapeLayer && LandscapeLayerName != NAME_None))
+			if (bReapplyNormal || Settings->ReapplyGroundSlope || Settings->ReapplyVertexColorMask || (Settings->ReapplyLandscapeLayer && LandscapeLayersValid(Settings)))
 			{
 				FHitResult Hit;
 				static const FName NAME_ReapplyInstancesForBrush = TEXT("ReapplyInstancesForBrush");
@@ -930,30 +934,20 @@ void FEdModeFoliage::ReapplyInstancesForBrush(UWorld* InWorld, AInstancedFoliage
 					}
 
 					// Cull instances for the landscape layer
-					if (Settings->ReapplyLandscapeLayer && LandscapeLayerName != NAME_None)
+					if (Settings->ReapplyLandscapeLayer && LandscapeLayersValid(Settings))
 					{
 						float HitWeight = 1.f;
-						ULandscapeHeightfieldCollisionComponent* HitLandscapeCollision = Cast<ULandscapeHeightfieldCollisionComponent>(Hit.Component.Get());
-						if (HitLandscapeCollision)
+						if (GetMaxHitWeight(Hit.Location, Hit.GetComponent(), Settings, &LandscapeLayerCaches, HitWeight))
 						{
-							ULandscapeComponent* HitLandscape = HitLandscapeCollision->RenderComponent.Get();
-							if (HitLandscape)
+							if (FilterByWeight(HitWeight, Settings))
 							{
-								TArray<uint8>* LayerCache = &LandscapeLayerCache->FindOrAdd(HitLandscape);
-								// TODO: FName to LayerInfo?
-								HitWeight = HitLandscape->GetLayerWeightAtLocation(Hit.Location, HitLandscape->GetLandscapeInfo()->GetLayerInfoByName(LandscapeLayerName), LayerCache);
-
-								// Reject instance randomly in proportion to weight
-								if (HitWeight <= FMath::FRand())
+								InstancesToDelete.Add(InstanceIndex);
+								if (bReapplyLocation)
 								{
-									InstancesToDelete.Add(InstanceIndex);
-									if (bReapplyLocation)
-									{
-										// restore the location so the hash removal will succeed
-										Instance.Location = OldInstanceLocation;
-									}
-									continue;
+									// restore the location so the hash removal will succeed
+									Instance.Location = OldInstanceLocation;
 								}
+								continue;
 							}
 						}
 					}
diff --git a/Engine/Source/Editor/FoliageEdit/Private/SFoliageEditMeshDisplayItem.cpp b/Engine/Source/Editor/FoliageEdit/Private/SFoliageEditMeshDisplayItem.cpp
--- a/Engine/Source/Editor/FoliageEdit/Private/SFoliageEditMeshDisplayItem.cpp
+++ b/Engine/Source/Editor/FoliageEdit/Private/SFoliageEditMeshDisplayItem.cpp
@@ -2439,12 +2439,23 @@ ECheckBoxState SFoliageEditMeshDisplayItem::IsLandscapeLayerReapplyChecked() con
 
 void SFoliageEditMeshDisplayItem::OnLandscapeLayerChanged(const FText& InValue)
 {
-	FoliageSettingsPtr->LandscapeLayer = FName(*InValue.ToString());
+	if (FoliageSettingsPtr->LandscapeLayers.Num() == 0)
+	{
+		FoliageSettingsPtr->LandscapeLayers.AddUninitialized(1);
+	}
+
+	FoliageSettingsPtr->LandscapeLayers[0] = FName(*InValue.ToString());
 }
 
 FText SFoliageEditMeshDisplayItem::GetLandscapeLayer() const
 {
-	return FText::FromName(FoliageSettingsPtr->LandscapeLayer);
+	FName LayerName = NAME_None;
+	if (FoliageSettingsPtr->LandscapeLayers.Num())
+	{
+		LayerName = FoliageSettingsPtr->LandscapeLayers[0];
+	}
+	
+	return FText::FromName(LayerName);
 }
 
 void SFoliageEditMeshDisplayItem::OnCollisionWithWorld(ECheckBoxState InState)
diff --git a/Engine/Source/Runtime/Foliage/Private/InstancedFoliage.cpp b/Engine/Source/Runtime/Foliage/Private/InstancedFoliage.cpp
--- a/Engine/Source/Runtime/Foliage/Private/InstancedFoliage.cpp
+++ b/Engine/Source/Runtime/Foliage/Private/InstancedFoliage.cpp
@@ -195,7 +195,6 @@ UFoliageType::UFoliageType(const FObjectInitializer& ObjectInitializer)
 	HeightMax = 262144.0f;
 	ZOffsetMin = 0.0f;
 	ZOffsetMax = 0.0f;
-	LandscapeLayer = NAME_None;
 	MinimumLayerWeight = 0.5f;
 	DisplayOrder = 0;
 	IsSelected = false;
@@ -246,6 +245,17 @@ UFoliageType::UFoliageType(const FObjectInitializer& ObjectInitializer)
 	UpdateGuid = FGuid::NewGuid();
 }
 
+void UFoliageType::Serialize(FArchive& Ar)
+{
+	Super::Serialize(Ar);
+
+	if (LandscapeLayer_DEPRECATED != NAME_None && LandscapeLayers.Num() == 0)	//we now store an array of names so initialize the array with the old name
+	{
+		LandscapeLayers.Add(LandscapeLayer_DEPRECATED);
+		LandscapeLayer_DEPRECATED = NAME_None;
+	}
+}
+
 
 UFoliageType_InstancedStaticMesh::UFoliageType_InstancedStaticMesh(const FObjectInitializer& ObjectInitializer)
 	: Super(ObjectInitializer)
diff --git a/Engine/Source/Runtime/Foliage/Public/FoliageType.h b/Engine/Source/Runtime/Foliage/Public/FoliageType.h
--- a/Engine/Source/Runtime/Foliage/Public/FoliageType.h
+++ b/Engine/Source/Runtime/Foliage/Public/FoliageType.h
@@ -77,8 +77,11 @@ class UFoliageType : public UObject
 	float HeightMax;
 
 	UPROPERTY(EditAnywhere, Category = General, meta = (Subcategory = "Placement"))
-	FName LandscapeLayer;
+	TArray<FName> LandscapeLayers;
 
+	UPROPERTY()
+	FName LandscapeLayer_DEPRECATED;
+	
 	UPROPERTY(EditAnywhere, Category = General, meta = (Subcategory = "Placement"))
 	float MinimumLayerWeight;
 
@@ -308,6 +311,8 @@ class UFoliageType : public UObject
 	FOLIAGE_API float GetInitAge(FRandomStream& RandomStream) const;
 	FOLIAGE_API float GetNextAge(const float CurrentAge, const int32 NumSteps) const;
 
+	virtual void Serialize(FArchive& Ar) override;
+
 private:
 
 	/** The curve used to interpolate the instance scale.*/
