Notify the rest of the editor when deselecting all on entering foliage edit mode.

diff --git a/Engine/Source/Editor/FoliageEdit/Private/FoliageEdMode.cpp b/Engine/Source/Editor/FoliageEdit/Private/FoliageEdMode.cpp
--- a/Engine/Source/Editor/FoliageEdit/Private/FoliageEdMode.cpp
+++ b/Engine/Source/Editor/FoliageEdit/Private/FoliageEdMode.cpp
@@ -93,7 +93,7 @@ void FEdModeFoliage::Enter()
 	FEdMode::Enter();
 
 	// Clear any selection in case the instanced foliage actor is selected
-	GEditor->SelectNone(false, true);
+	GEditor->SelectNone(true, true);
 
 	// Load UI settings from config file
 	UISettings.Load();
