Foliage: do not allow cross-level painting with non-shared foliage settings

diff --git a/Engine/Source/Editor/FoliageEdit/Private/FoliageEdMode.cpp b/Engine/Source/Editor/FoliageEdit/Private/FoliageEdMode.cpp
--- a/Engine/Source/Editor/FoliageEdit/Private/FoliageEdMode.cpp
+++ b/Engine/Source/Editor/FoliageEdit/Private/FoliageEdMode.cpp
@@ -525,6 +525,15 @@ void FEdModeFoliage::CalculatePotentialInstances(UWorld* InWorld, const AInstanc
 			{
 				continue;
 			}
+
+			// We can paint into new level only if FoliageType is shared
+			if (DesiredInst.PlacementMode == EFoliagePlacementMode::Manual)
+			{
+				if (!Settings->IsAsset() && Settings->GetOuter() != TargetIFA)
+				{
+					continue;
+				}
+			}
 				
 			FFoliageMeshInfo* MeshInfo = DesiredInst.MeshInfo;
 			//We have to add the mesh into the foliage actor. This assumes the foliage type is just a concrete InstancedStaticMesh.
