Start breaking up code into thread safe parts for upcoming parallelization effort. Made it so procedural foliage's FoliageType overrides any settings specified by manual painting

diff --git a/Engine/Source/Editor/FoliageEdit/Private/FoliageEdMode.cpp b/Engine/Source/Editor/FoliageEdit/Private/FoliageEdMode.cpp
--- a/Engine/Source/Editor/FoliageEdit/Private/FoliageEdMode.cpp
+++ b/Engine/Source/Editor/FoliageEdit/Private/FoliageEdMode.cpp
@@ -33,6 +33,7 @@
 #include "LandscapeDataAccess.h"
 #include "FoliageType_InstancedStaticMesh.h"
 
+
 #define LOCTEXT_NAMESPACE "FoliageEdMode"
 #define FOLIAGE_SNAP_TRACE (10000.f)
 
@@ -344,9 +345,9 @@ void FEdModeFoliage::GetRandomVectorInBrush(FVector& OutStart, FVector& OutEnd)
 	OutEnd = BrushLocation + UISettings.GetRadius() * (Ru * U + Rv * V + Rw * BrushTraceDirection);
 }
 
-static bool CheckLocationForPotentialInstance(FFoliageMeshInfo& MeshInfo, const UFoliageType* Settings, float DensityCheckRadius, const FVector& Location, const FVector& Normal, TArray<FVector>& PotentialInstanceLocations, FFoliageInstanceHash& PotentialInstanceHash)
+/** This does not check for overlaps or density */
+static bool CheckLocationForPotentialInstance_ThreadSafe(const UFoliageType* Settings, const FVector& Location, const FVector& Normal)
 {
-
 	// Check height range
 	if (Location.Z < Settings->HeightMin || Location.Z > Settings->HeightMax)
 	{
@@ -361,6 +362,17 @@ static bool CheckLocationForPotentialInstance(FFoliageMeshInfo& MeshInfo, const
 		return false;
 	}
 
+	return true;
+}
+
+static bool CheckLocationForPotentialInstance(FFoliageMeshInfo& MeshInfo, const UFoliageType* Settings, float DensityCheckRadius, const FVector& Location, const FVector& Normal, TArray<FVector>& PotentialInstanceLocations, FFoliageInstanceHash& PotentialInstanceHash)
+{
+
+	if (CheckLocationForPotentialInstance_ThreadSafe(Settings, Location, Normal) == false)
+	{
+		return false;
+	}
+
 	// Check existing instances. Use the Density radius rather than the minimum radius
 	if (MeshInfo.CheckForOverlappingSphere(FSphere(Location, DensityCheckRadius)))
 	{
@@ -452,7 +464,7 @@ bool GetMaxHitWeight(const FVector& Location, UActorComponent* Component, const
 	{
 		if (ULandscapeComponent* HitLandscape = HitLandscapeCollision->RenderComponent.Get())
 		{
-			for (FName LandscapeLayerName : Settings->LandscapeLayers)
+			for (const FName LandscapeLayerName : Settings->LandscapeLayers)
 			{
 				// Cache store mapping between component and weight data
 				TMap<ULandscapeComponent*, TArray<uint8> >* LandscapeLayerCache = &LandscapeLayerCaches->FindOrAdd(LandscapeLayerName);;
@@ -476,7 +488,103 @@ bool FilterByWeight(float Weight, const UFoliageType* Settings)
 	return Weight < FMath::Max(SMALL_NUMBER, WeightNeeded);
 }
 
-void FEdModeFoliage::CalculatePotentialInstances(UWorld* InWorld, const AInstancedFoliageActor* IgnoreIFA, const UFoliageType* Settings, const TArray<FDesiredFoliageInstance>& DesiredInstances, TArray<FPotentialInstance> OutPotentialInstances[NUM_INSTANCE_BUCKETS], LandscapeLayerCacheData* LandscapeLayerCachesPtr, const FFoliageUISettings* UISettings)
+bool GeometryFilterCheck(const FHitResult& Hit, const UWorld* InWorld, const FDesiredFoliageInstance& DesiredInst, const FFoliageUISettings* UISettings)
+{
+	if (UPrimitiveComponent* PrimComp = Hit.Component.Get())
+	{
+		if (DesiredInst.PlacementMode == EFoliagePlacementMode::Manual)
+		{
+			UMaterialInterface* Material = PrimComp ? PrimComp->GetMaterial(0) : nullptr;
+			if (PrimComp->GetOutermost() != InWorld->GetCurrentLevel()->GetOutermost() ||
+				(!UISettings->bFilterLandscape && PrimComp->IsA(ULandscapeHeightfieldCollisionComponent::StaticClass())) ||
+				(!UISettings->bFilterStaticMesh && PrimComp->IsA(UStaticMeshComponent::StaticClass())) ||
+				(!UISettings->bFilterBSP && PrimComp->IsA(UModelComponent::StaticClass())) ||
+				(!UISettings->bFilterTranslucent && Material && IsTranslucentBlendMode(Material->GetBlendMode()))
+				)
+			{
+				return false;
+			}
+		}
+	}
+	else
+	{
+		return false;
+	}
+
+	return true;
+}
+
+bool FEdModeFoliage::VertexMaskCheck(const FHitResult& Hit, const UFoliageType* Settings)
+{
+	if (Settings->VertexColorMask != FOLIAGEVERTEXCOLORMASK_Disabled && Hit.FaceIndex != INDEX_NONE)
+	{
+		if (UStaticMeshComponent* HitStaticMeshComponent = Cast<UStaticMeshComponent>(Hit.Component.Get()))
+		{
+			FColor VertexColor;
+			if (FEdModeFoliage::GetStaticMeshVertexColorForHit(HitStaticMeshComponent, Hit.FaceIndex, Hit.ImpactPoint, VertexColor))
+			{
+				if (!CheckVertexColor(Settings, VertexColor))
+				{
+					return false;
+				}
+			}
+		}
+	}
+
+	return true;
+}
+
+bool LandscapeLayerCheck(const FHitResult& Hit, const UFoliageType* Settings, FEdModeFoliage::LandscapeLayerCacheData& LandscapeLayersCache, float& OutHitWeight)
+{
+	OutHitWeight = 1.f;
+	if (LandscapeLayersValid(Settings) && GetMaxHitWeight(Hit.ImpactPoint, Hit.Component.Get(), Settings, &LandscapeLayersCache, OutHitWeight))
+	{
+		// Reject instance randomly in proportion to weight
+		if (FilterByWeight(OutHitWeight, Settings))
+		{
+			return false;
+		}
+	}
+
+	return true;
+}
+
+void FEdModeFoliage::CalculatePotentialInstances_ThreadSafe(const UWorld* InWorld, const AInstancedFoliageActor* IgnoreIFA, const UFoliageType* Settings, const TArray<FDesiredFoliageInstance>* DesiredInstances, TArray<FPotentialInstance> OutPotentialInstances[NUM_INSTANCE_BUCKETS], const FFoliageUISettings* UISettings, const int32 StartIdx, const int32 LastIdx)
+{
+	LandscapeLayerCacheData LocalCache;
+
+	// Reserve space in buckets for a potential instances 
+	for (int32 BucketIdx = 0; BucketIdx < NUM_INSTANCE_BUCKETS; ++BucketIdx)
+	{
+		auto& Bucket = OutPotentialInstances[BucketIdx];
+		Bucket.Reserve(DesiredInstances->Num());
+	}
+
+	for (int32 InstanceIdx = StartIdx; InstanceIdx <= LastIdx; ++InstanceIdx)
+	{
+		const FDesiredFoliageInstance& DesiredInst = (*DesiredInstances)[InstanceIdx];
+		FHitResult Hit;
+		static FName NAME_AddFoliageInstances = FName(TEXT("AddFoliageInstances"));
+		if (AInstancedFoliageActor::FoliageTrace(InWorld, Hit, DesiredInst, IgnoreIFA, NAME_AddFoliageInstances, true))
+		{
+			float HitWeight = 1.f;
+			const bool bValidInstance = GeometryFilterCheck(Hit, InWorld, DesiredInst, UISettings)
+										&& CheckLocationForPotentialInstance_ThreadSafe(Settings, Hit.ImpactPoint, Hit.ImpactNormal)
+										&& VertexMaskCheck(Hit, Settings)
+										&& LandscapeLayerCheck(Hit, Settings, LocalCache, HitWeight);
+
+			
+
+			if (bValidInstance)
+			{
+				const int32 BucketIndex = FMath::RoundToInt(HitWeight * (float)(NUM_INSTANCE_BUCKETS - 1));
+				OutPotentialInstances[BucketIndex].Add(FPotentialInstance(Hit.ImpactPoint, Hit.ImpactNormal, Hit.Component.Get(), HitWeight, nullptr, nullptr, DesiredInst));
+			}
+		}
+	}
+}
+
+void FEdModeFoliage::CalculatePotentialInstances(const UWorld* InWorld, const AInstancedFoliageActor* IgnoreIFA, const UFoliageType* Settings, const TArray<FDesiredFoliageInstance>& DesiredInstances, TArray<FPotentialInstance> OutPotentialInstances[NUM_INSTANCE_BUCKETS], LandscapeLayerCacheData* LandscapeLayerCachesPtr, const FFoliageUISettings* UISettings)
 {
 	LandscapeLayerCacheData LocalCache;
 	LandscapeLayerCachesPtr = LandscapeLayerCachesPtr ? LandscapeLayerCachesPtr : &LocalCache;
@@ -502,22 +610,14 @@ void FEdModeFoliage::CalculatePotentialInstances(UWorld* InWorld, const AInstanc
 		static FName NAME_AddFoliageInstances = FName(TEXT("AddFoliageInstances"));
 		if (AInstancedFoliageActor::FoliageTrace(InWorld, Hit, DesiredInst, IgnoreIFA, NAME_AddFoliageInstances, true))
 		{
-			// Check filters
-			UPrimitiveComponent* InstanceBase = Hit.Component.Get();
-			check(InstanceBase);
-
-			if (DesiredInst.PlacementMode == EFoliagePlacementMode::Manual)
+			float HitWeight = 1.f;
+			if (GeometryFilterCheck(Hit, InWorld, DesiredInst, UISettings) == false)
 			{
-				UMaterialInterface* Material = InstanceBase->GetMaterial(0);
-				if ((!UISettings->bFilterLandscape && InstanceBase->IsA(ULandscapeHeightfieldCollisionComponent::StaticClass())) ||
-					(!UISettings->bFilterStaticMesh && InstanceBase->IsA(UStaticMeshComponent::StaticClass())) ||
-					(!UISettings->bFilterBSP && InstanceBase->IsA(UModelComponent::StaticClass())) ||
-					(!UISettings->bFilterTranslucent && Material && IsTranslucentBlendMode(Material->GetBlendMode()))
-					)
-				{
-					continue;
-				}
+				continue;
 			}
+			
+			UPrimitiveComponent* InstanceBase = Hit.GetComponent();
+			check(InstanceBase);
 
 			// New foliage instances always go into base component level
 			AInstancedFoliageActor* TargetIFA = AInstancedFoliageActor::GetInstancedFoliageActorForLevel(InstanceBase->GetComponentLevel(), true);
@@ -530,11 +630,11 @@ void FEdModeFoliage::CalculatePotentialInstances(UWorld* InWorld, const AInstanc
 			if (DesiredInst.PlacementMode == EFoliagePlacementMode::Manual)
 			{
 				if (!Settings->IsAsset() && Settings->GetOuter() != TargetIFA)
-				{
-					continue;
-				}
+			{
+				continue;
+			}
 			}
-				
+
 			FFoliageMeshInfo* MeshInfo = DesiredInst.MeshInfo;
 			//We have to add the mesh into the foliage actor. This assumes the foliage type is just a concrete InstancedStaticMesh.
 			//With the new tree system coming this will not be true and you'll need to traverse the tree and insert the needed meshes
@@ -554,45 +654,15 @@ void FEdModeFoliage::CalculatePotentialInstances(UWorld* InWorld, const AInstanc
 				}
 			}
 
-			if (MeshInfo == nullptr)
-			{
-				continue;
-			}
-
-			if (CheckLocationForPotentialInstance(*MeshInfo, Settings, DensityCheckRadius, Hit.ImpactPoint, Hit.ImpactNormal, PotentialInstanceLocations, PotentialInstanceHash) == false)
+			if (MeshInfo)
 			{
-				continue;
-			}
+				const bool bValidInstance =	CheckLocationForPotentialInstance(*MeshInfo, Settings, DensityCheckRadius, Hit.ImpactPoint, Hit.ImpactNormal, PotentialInstanceLocations, PotentialInstanceHash)
+											&& VertexMaskCheck(Hit, Settings)
+											&& LandscapeLayerCheck(Hit, Settings, LocalCache, HitWeight);
 
-			// Check vertex color mask
-			if (Settings->VertexColorMask != FOLIAGEVERTEXCOLORMASK_Disabled && Hit.FaceIndex != INDEX_NONE)
-			{
-				if (UStaticMeshComponent* HitStaticMeshComponent = Cast<UStaticMeshComponent>(InstanceBase))
-				{
-					FColor VertexColor;
-					if (GetStaticMeshVertexColorForHit(HitStaticMeshComponent, Hit.FaceIndex, Hit.ImpactPoint, VertexColor))
-					{
-						if (!CheckVertexColor(Settings, VertexColor))
-						{
-							continue;
-						}
-					}
-				}
-			}
-
-			// Check landscape layer
-			float HitWeight = 1.f;
-			if (LandscapeLayersValid(Settings) && GetMaxHitWeight(Hit.ImpactPoint, InstanceBase, Settings, LandscapeLayerCachesPtr, HitWeight))
-			{
-				// Reject instance randomly in proportion to weight
-				if (FilterByWeight(HitWeight, Settings))
-				{
-					continue;
-				}
+				const int32 BucketIndex = FMath::RoundToInt(HitWeight * (float)(NUM_INSTANCE_BUCKETS - 1));
+				OutPotentialInstances[BucketIndex].Add(FPotentialInstance(Hit.ImpactPoint, Hit.ImpactNormal, InstanceBase, HitWeight, MeshInfo, TargetIFA, DesiredInst));
 			}
-
-			const int32 BucketIndex = FMath::RoundToInt(HitWeight * (float)(NUM_INSTANCE_BUCKETS - 1));
-			OutPotentialInstances[BucketIndex].Add(FPotentialInstance(Hit.ImpactPoint, Hit.ImpactNormal, InstanceBase, HitWeight, MeshInfo, TargetIFA, DesiredInst));
 		}
 	}
 }
@@ -615,10 +685,73 @@ void FEdModeFoliage::AddInstances(UWorld* InWorld, const TArray<FDesiredFoliageI
 	}
 }
 
+FFoliageMeshInfo* UpdateMeshSettings(AInstancedFoliageActor* IFA, const FDesiredFoliageInstance& DesiredInstance)
+{
+	FFoliageMeshInfo* MeshInfo = nullptr;
+	//We have to add the mesh into the foliage actor. This assumes the foliage type is just a concrete InstancedStaticMesh.
+	//With the new tree system coming this will not be true and you'll need to traverse the tree and insert the needed meshes
+	if (const UFoliageType_InstancedStaticMesh* Settings = Cast<UFoliageType_InstancedStaticMesh>(DesiredInstance.FoliageType))
+	{
+		if (UStaticMesh* StaticMesh = Settings->GetStaticMesh())
+		{
+			if (UFoliageType* ExistingSettings = IFA->GetSettingsForMesh(StaticMesh))
+			{
+				MeshInfo = IFA->UpdateMeshSettings(StaticMesh, Settings);
+			}
+			else
+			{
+				MeshInfo = IFA->AddMesh(Settings->GetStaticMesh(), nullptr, Settings);
+			}
+		}
+	}
+
+	return MeshInfo;
+}
+
 void FEdModeFoliage::AddInstancesImp(UWorld* InWorld, const UFoliageType* Settings, const TArray<FDesiredFoliageInstance>& DesiredInstances, const AInstancedFoliageActor* IgnoreIFA, const TArray<int32>& ExistingInstanceBuckets, const float Pressure, LandscapeLayerCacheData* LandscapeLayerCachesPtr, const FFoliageUISettings* UISettings)
 {
+	if (DesiredInstances.Num() == 0)
+	{
+		return;
+	}
+
+
 	TArray<FPotentialInstance> PotentialInstanceBuckets[NUM_INSTANCE_BUCKETS];
+	if (DesiredInstances[0].PlacementMode == EFoliagePlacementMode::Manual)
+	{
 	CalculatePotentialInstances(InWorld, IgnoreIFA, Settings, DesiredInstances, PotentialInstanceBuckets, LandscapeLayerCachesPtr, UISettings);
+	}
+	else
+	{
+		//@TODO: actual threaded part coming, need parts of this refactor sooner for content team
+		CalculatePotentialInstances_ThreadSafe(InWorld, IgnoreIFA, Settings, &DesiredInstances, PotentialInstanceBuckets, nullptr, 0, DesiredInstances.Num() - 1);
+
+		TMap<AInstancedFoliageActor*, FFoliageMeshInfo*> UpdatedIFAs;	//we want to override any existing mesh settings with the procedural settings.
+		for (auto& Bucket : PotentialInstanceBuckets)
+		{
+			for (FPotentialInstance& PotentialInst : Bucket)
+			{
+				// New foliage instances always go into base component level
+				AInstancedFoliageActor* TargetIFA = AInstancedFoliageActor::GetInstancedFoliageActorForLevel(PotentialInst.HitComponent->GetComponentLevel());
+				FFoliageMeshInfo* MeshInfo = nullptr;
+				if (FFoliageMeshInfo** MeshInfoResult = UpdatedIFAs.Find(TargetIFA))
+				{
+					MeshInfo = *MeshInfoResult;
+				}
+				else
+				{
+					MeshInfo = UpdateMeshSettings(TargetIFA, PotentialInst.DesiredInstance);
+					UpdatedIFAs.Add(TargetIFA, MeshInfo);
+				}
+
+				PotentialInst.IFA = TargetIFA;
+				PotentialInst.MeshInfo = MeshInfo;
+			}
+		}
+	}
+		
+	UStaticMesh* StaticMesh = Settings->GetStaticMesh();
+	
 
 	for (int32 BucketIdx = 0; BucketIdx < NUM_INSTANCE_BUCKETS; BucketIdx++)
 	{
@@ -1397,16 +1530,16 @@ void FEdModeFoliage::ApplyPaintBucket(AActor* Actor, bool bRemove)
 	}
 }
 
-bool FEdModeFoliage::GetStaticMeshVertexColorForHit(UStaticMeshComponent* InStaticMeshComponent, int32 InTriangleIndex, const FVector& InHitLocation, FColor& OutVertexColor)
+bool FEdModeFoliage::GetStaticMeshVertexColorForHit(const UStaticMeshComponent* InStaticMeshComponent, int32 InTriangleIndex, const FVector& InHitLocation, FColor& OutVertexColor)
 {
-	UStaticMesh* StaticMesh = InStaticMeshComponent->StaticMesh;
+	const UStaticMesh* StaticMesh = InStaticMeshComponent->StaticMesh;
 
 	if (StaticMesh == nullptr || StaticMesh->RenderData == nullptr)
 	{
 		return false;
 	}
 
-	FStaticMeshLODResources& LODModel = StaticMesh->RenderData->LODResources[0];
+	const FStaticMeshLODResources& LODModel = StaticMesh->RenderData->LODResources[0];
 	bool bHasInstancedColorData = false;
 	const FStaticMeshComponentLODInfo* InstanceMeshLODInfo = nullptr;
 	if (InStaticMeshComponent->LODData.Num() > 0)
@@ -1911,7 +2044,7 @@ bool FEdModeFoliage::InputKey(FEditorViewportClient* ViewportClient, FViewport*
 					for (int32 InstanceIndex : SelectedIndices)
 					{
 						bMovedInstance|= SnapInstanceToGround(IFA, MeshPair.Key->AlignMaxAngle, Mesh, InstanceIndex);
-					}
+								}
 					Mesh.PostMoveInstances(IFA, SelectedIndices);
 				}
 
diff --git a/Engine/Source/Editor/FoliageEdit/Private/FoliageEdMode.h b/Engine/Source/Editor/FoliageEdit/Private/FoliageEdMode.h
--- a/Engine/Source/Editor/FoliageEdit/Private/FoliageEdMode.h
+++ b/Engine/Source/Editor/FoliageEdit/Private/FoliageEdMode.h
@@ -318,7 +318,11 @@ private:
 	/** Common code for adding instances to world based on settings */
 	static void AddInstancesImp(UWorld* InWorld, const UFoliageType* Settings, const TArray<FDesiredFoliageInstance>& DesiredInstances, const AInstancedFoliageActor* IgnoreIFA = nullptr, const TArray<int32>& ExistingInstances = TArray<int32>(), const float Pressure = 1.f, LandscapeLayerCacheData* LandscapeLayerCaches = nullptr, const FFoliageUISettings* UISettings = nullptr);
 
-	static void CalculatePotentialInstances(UWorld* InWorld, const AInstancedFoliageActor* IFA, const UFoliageType* Settings, const TArray<FDesiredFoliageInstance>& DesiredInstances, TArray<FPotentialInstance> OutPotentialInstances[NUM_INSTANCE_BUCKETS], LandscapeLayerCacheData* LandscaleLayerCachesPtr, const FFoliageUISettings* UISettings);
+	/** Logic for determining which instances can be placed in the world*/
+	static void CalculatePotentialInstances(const UWorld* InWorld, const AInstancedFoliageActor* IFA, const UFoliageType* Settings, const TArray<FDesiredFoliageInstance>& DesiredInstances, TArray<FPotentialInstance> OutPotentialInstances[NUM_INSTANCE_BUCKETS], LandscapeLayerCacheData* LandscaleLayerCachesPtr, const FFoliageUISettings* UISettings);
+
+	/** Similar to CalculatePotentialInstances, but it doesn't do any overlap checks which are much harder to thread. Meant to be run in parallel for placing lots of instances */
+	static void CalculatePotentialInstances_ThreadSafe(const UWorld* InWorld, const AInstancedFoliageActor* IgnoreIFA, const UFoliageType* Settings, const TArray<FDesiredFoliageInstance>* DesiredInstances, TArray<FPotentialInstance> OutPotentialInstances[NUM_INSTANCE_BUCKETS], const FFoliageUISettings* UISettings, const int32 StartIdx, const int32 LastIdx);
 
 	/** Remove instances inside the brush to match DesiredInstanceCount. NOTE: PotentialInstancesToRemove array is modified by this function. */
 	void RemoveInstancesForBrush(AInstancedFoliageActor* IFA, FFoliageMeshInfo& MeshInfo, int32 DesiredInstanceCount, TArray<int32>& PotentialInstancesToRemove, float Pressure);
@@ -327,7 +331,10 @@ private:
 	void ReapplyInstancesForBrush(UWorld* InWorld, AInstancedFoliageActor* IFA, UFoliageType* Settings, const TArray<int32>& ExistingInstances);
 
 	/** Lookup the vertex color corresponding to a location traced on a static mesh */
-	static bool GetStaticMeshVertexColorForHit(UStaticMeshComponent* InStaticMeshComponent, int32 InTriangleIndex, const FVector& InHitLocation, FColor& OutVertexColor);
+	static bool GetStaticMeshVertexColorForHit(const UStaticMeshComponent* InStaticMeshComponent, int32 InTriangleIndex, const FVector& InHitLocation, FColor& OutVertexColor);
+
+	/** Does a filter based on the vertex color of a static mesh */
+	static bool VertexMaskCheck(const FHitResult& Hit, const UFoliageType* Settings);
 
 	/** 
 	 * Snaps given instance to the ground  
diff --git a/Engine/Source/Runtime/Foliage/Private/InstancedFoliage.cpp b/Engine/Source/Runtime/Foliage/Private/InstancedFoliage.cpp
--- a/Engine/Source/Runtime/Foliage/Private/InstancedFoliage.cpp
+++ b/Engine/Source/Runtime/Foliage/Private/InstancedFoliage.cpp
@@ -478,6 +478,28 @@
 #endif
 }
 
+void FFoliageMeshInfo::UpdateComponentSettings(const UFoliageType* InSettings)
+{
+	if (Component)
+	{
+		Component->Mobility = InSettings->bEnableStaticLighting ? EComponentMobility::Static : EComponentMobility::Movable;
+		Component->InstanceStartCullDistance = InSettings->StartCullDistance;
+		Component->InstanceEndCullDistance = InSettings->EndCullDistance;
+
+		Component->CastShadow = InSettings->CastShadow;
+		Component->bCastDynamicShadow = InSettings->bCastDynamicShadow;
+		Component->bCastStaticShadow = InSettings->bCastStaticShadow;
+		Component->bAffectDynamicIndirectLighting = InSettings->bAffectDynamicIndirectLighting;
+		Component->bAffectDistanceFieldLighting = InSettings->bAffectDistanceFieldLighting;
+		Component->bCastShadowAsTwoSided = InSettings->bCastShadowAsTwoSided;
+		Component->bReceivesDecals = InSettings->bReceivesDecals;
+		Component->bOverrideLightMapRes = InSettings->bOverrideLightMapRes;
+		Component->OverriddenLightMapRes = InSettings->OverriddenLightMapRes;
+
+		Component->BodyInstance.CopyBodyInstancePropertiesFrom(&InSettings->BodyInstance);
+	}
+}
+
 void FFoliageMeshInfo::AddInstance(AInstancedFoliageActor* InIFA, const UFoliageType* InSettings, const FFoliageInstance& InNewInstance, UActorComponent* InBaseComponent)
 {
 	FFoliageInstance Instance = InNewInstance;
@@ -493,26 +515,12 @@
 	{
 		Component = ConstructObject<UHierarchicalInstancedStaticMeshComponent>(UHierarchicalInstancedStaticMeshComponent::StaticClass(), InIFA, NAME_None, RF_Transactional);
 
-		Component->Mobility = InSettings->bEnableStaticLighting ? EComponentMobility::Static : EComponentMobility::Movable;
-
 		Component->StaticMesh = InSettings->GetStaticMesh();
 		Component->bSelectable = true;
 		Component->bHasPerInstanceHitProxies = true;
 		Component->InstancingRandomSeed = FMath::Rand();
-		Component->InstanceStartCullDistance = InSettings->StartCullDistance;
-		Component->InstanceEndCullDistance = InSettings->EndCullDistance;
 
-		Component->CastShadow = InSettings->CastShadow;
-		Component->bCastDynamicShadow = InSettings->bCastDynamicShadow;
-		Component->bCastStaticShadow = InSettings->bCastStaticShadow;
-		Component->bAffectDynamicIndirectLighting = InSettings->bAffectDynamicIndirectLighting;
-		Component->bAffectDistanceFieldLighting = InSettings->bAffectDistanceFieldLighting;
-		Component->bCastShadowAsTwoSided = InSettings->bCastShadowAsTwoSided;
-		Component->bReceivesDecals = InSettings->bReceivesDecals;
-		Component->bOverrideLightMapRes = InSettings->bOverrideLightMapRes;
-		Component->OverriddenLightMapRes = InSettings->OverriddenLightMapRes;
+		UpdateComponentSettings(InSettings);
 
-		Component->BodyInstance.CopyBodyInstancePropertiesFrom(&InSettings->BodyInstance);
-
 		Component->AttachTo(InIFA->GetRootComponent());
 		Component->RegisterComponent();
@@ -1251,28 +1259,9 @@
 	return MeshInfo;
 }
 
-FFoliageMeshInfo* AInstancedFoliageActor::AddMesh(UStaticMesh* InMesh, UFoliageType** OutSettings, const UFoliageType_InstancedStaticMesh* DefaultSettings)
+
+void UpdateSettingsBounds(const UStaticMesh* InMesh, UFoliageType_InstancedStaticMesh* Settings)
 {
-	check(GetSettingsForMesh(InMesh) == nullptr);
-
-	MarkPackageDirty();
-
-	UFoliageType_InstancedStaticMesh* Settings = nullptr;
-#if WITH_EDITORONLY_DATA
-	if (DefaultSettings)
-	{
-		// TODO: Can't we just use this directly?
-		Settings = DuplicateObject<UFoliageType_InstancedStaticMesh>(DefaultSettings, this);
-	}
-	else
-#endif
-	{
-		Settings = ConstructObject<UFoliageType_InstancedStaticMesh>(UFoliageType_InstancedStaticMesh::StaticClass(), this);
-	}
-	Settings->Mesh = InMesh;
-
-	FFoliageMeshInfo* MeshInfo = AddMesh(Settings);
-
 	const FBoxSphereBounds MeshBounds = InMesh->GetBounds();
 
 	Settings->MeshBounds = MeshBounds;
@@ -1301,6 +1290,52 @@
 	}
 
 	Settings->LowBoundOriginRadius = FVector((MinX + MaxX), (MinY + MaxY), FMath::Sqrt(FMath::Square(MaxX - MinX) + FMath::Square(MaxY - MinY))) * 0.5f;
+}
 
+
+#if WITH_EDITORONLY_DATA
+FFoliageMeshInfo* AInstancedFoliageActor::UpdateMeshSettings(const UStaticMesh* InMesh, const UFoliageType_InstancedStaticMesh* DefaultSettings)
+{
+	if (UFoliageType* OldSettings = GetSettingsForMesh(InMesh))
+	{
+		MarkPackageDirty();
+
+		UFoliageType_InstancedStaticMesh* NewSettings = DuplicateObject<UFoliageType_InstancedStaticMesh>(DefaultSettings, this);
+		UpdateSettingsBounds(InMesh, NewSettings);
+
+		TUniqueObj<FFoliageMeshInfo> MeshInfo;
+		FoliageMeshes.RemoveAndCopyValue(OldSettings, MeshInfo);
+		MeshInfo->FoliageTypeUpdateGuid = NewSettings->UpdateGuid;
+		MeshInfo->UpdateComponentSettings(NewSettings);
+		return &*FoliageMeshes.Add(NewSettings, MoveTemp(MeshInfo));
+	}
+
+	return nullptr;
+}
+#endif
+
+FFoliageMeshInfo* AInstancedFoliageActor::AddMesh(UStaticMesh* InMesh, UFoliageType** OutSettings, const UFoliageType_InstancedStaticMesh* DefaultSettings)
+{
+	check(GetSettingsForMesh(InMesh) == nullptr);
+
+	MarkPackageDirty();
+
+	UFoliageType_InstancedStaticMesh* Settings = nullptr;
+#if WITH_EDITORONLY_DATA
+	if (DefaultSettings)
+	{
+		// TODO: Can't we just use this directly?
+		Settings = DuplicateObject<UFoliageType_InstancedStaticMesh>(DefaultSettings, this);
+	}
+	else
+#endif
+	{
+		Settings = ConstructObject<UFoliageType_InstancedStaticMesh>(UFoliageType_InstancedStaticMesh::StaticClass(), this);
+	}
+	Settings->Mesh = InMesh;
+
+	FFoliageMeshInfo* MeshInfo = AddMesh(Settings);
+	UpdateSettingsBounds(InMesh, Settings);
+
 	if (OutSettings)
 	{
@@ -1913,7 +1948,7 @@
 	}
 }
 
-bool AInstancedFoliageActor::FoliageTrace(UWorld* InWorld, FHitResult& OutHit, const FDesiredFoliageInstance& DesiredInstance, const AInstancedFoliageActor* IgnoreIFA, FName InTraceTag, bool InbReturnFaceIndex)
+bool AInstancedFoliageActor::FoliageTrace(const UWorld* InWorld, FHitResult& OutHit, const FDesiredFoliageInstance& DesiredInstance, const AInstancedFoliageActor* IgnoreIFA, FName InTraceTag, bool InbReturnFaceIndex)
 {
 	FCollisionQueryParams QueryParams(InTraceTag, true, IgnoreIFA);
 	QueryParams.bReturnFaceIndex = InbReturnFaceIndex;
@@ -1931,7 +1966,7 @@
 		{
 			if (Hit.Actor.IsValid())	//if we hit the ProceduralFoliage blocking volume don't spawn instance
 			{
-				if (AProceduralFoliageBlockingVolume* ProceduralFoliageBlockingVolume = Cast<AProceduralFoliageBlockingVolume>(Hit.Actor.Get()))
+				if (const AProceduralFoliageBlockingVolume* ProceduralFoliageBlockingVolume = Cast<AProceduralFoliageBlockingVolume>(Hit.Actor.Get()))
 				{
 					const AProceduralFoliageActor* ProceduralFoliageActor = ProceduralFoliageBlockingVolume->ProceduralFoliageActor;
 					if (ProceduralFoliageActor == nullptr || ProceduralFoliageActor->ProceduralComponent == nullptr || ProceduralFoliageActor->ProceduralComponent->GetProceduralGuid() == DesiredInstance.ProceduralGuid)
@@ -1947,7 +1982,7 @@
 		}
 			
 		// In the editor traces can hit "No Collision" type actors, so ugh.
-		FBodyInstance* BodyInstance = Hit.Component->GetBodyInstance();
+		const FBodyInstance* BodyInstance = Hit.Component->GetBodyInstance();
 		if (BodyInstance->GetCollisionEnabled() != ECollisionEnabled::QueryAndPhysics || BodyInstance->GetResponseToChannel(ECC_WorldStatic) != ECR_Block)
 		{
 			continue;
diff --git a/Engine/Source/Runtime/Foliage/Public/InstancedFoliage.h b/Engine/Source/Runtime/Foliage/Public/InstancedFoliage.h
--- a/Engine/Source/Runtime/Foliage/Public/InstancedFoliage.h
+++ b/Engine/Source/Runtime/Foliage/Public/InstancedFoliage.h
@@ -187,6 +187,7 @@ struct FFoliageMeshInfo
 	FOLIAGE_API void AddInstance(AInstancedFoliageActor* InIFA, const UFoliageType* InSettings, const FFoliageInstance& InNewInstance);
 	FOLIAGE_API void AddInstance(AInstancedFoliageActor* InIFA, const UFoliageType* InSettings, const FFoliageInstance& InNewInstance, UActorComponent* InBaseComponent);
 	FOLIAGE_API void RemoveInstances(AInstancedFoliageActor* InIFA, const TArray<int32>& InInstancesToRemove);
+	FOLIAGE_API void UpdateComponentSettings(const UFoliageType* InSettings);
 	FOLIAGE_API void PreMoveInstances(AInstancedFoliageActor* InIFA, const TArray<int32>& InInstancesToMove);
 	FOLIAGE_API void PostMoveInstances(AInstancedFoliageActor* InIFA, const TArray<int32>& InInstancesMoved);
 	FOLIAGE_API void PostUpdateInstances(AInstancedFoliageActor* InIFA, const TArray<int32>& InInstancesUpdated, bool bReAddToHash = false);
diff --git a/Engine/Source/Runtime/Foliage/Public/InstancedFoliageActor.h b/Engine/Source/Runtime/Foliage/Public/InstancedFoliageActor.h
--- a/Engine/Source/Runtime/Foliage/Public/InstancedFoliageActor.h
+++ b/Engine/Source/Runtime/Foliage/Public/InstancedFoliageActor.h
@@ -86,7 +86,7 @@ public:
 	*/
 	static FOLIAGE_API AInstancedFoliageActor* GetInstancedFoliageActorForLevel(ULevel* Level, bool bCreateIfNone = true);
 
-	static FOLIAGE_API bool FoliageTrace(UWorld* InWorld, FHitResult& OutHit, const FDesiredFoliageInstance& DesiredInstance, const AInstancedFoliageActor* IgnoreIFA = nullptr, FName InTraceTag = NAME_None, bool InbReturnFaceIndex = false);
+	static FOLIAGE_API bool FoliageTrace(const UWorld* InWorld, FHitResult& OutHit, const FDesiredFoliageInstance& DesiredInstance, const AInstancedFoliageActor* IgnoreIFA = nullptr, FName InTraceTag = NAME_None, bool InbReturnFaceIndex = false);
 	FOLIAGE_API bool CheckCollisionWithWorld(const UFoliageType* Settings, const FFoliageInstance& Inst, const FVector& HitNormal, const FVector& HitLocation);
 
 #if WITH_EDITOR
@@ -120,6 +120,7 @@ public:
 	// Add a new static mesh.
 	FOLIAGE_API FFoliageMeshInfo* AddMesh(UStaticMesh* InMesh, UFoliageType** OutSettings = nullptr, const UFoliageType_InstancedStaticMesh* DefaultSettings = nullptr);
 	FOLIAGE_API FFoliageMeshInfo* AddMesh(UFoliageType* InType);
+	FOLIAGE_API FFoliageMeshInfo* UpdateMeshSettings(const UStaticMesh* InMesh, const UFoliageType_InstancedStaticMesh* DefaultSettings);
 
 	// Remove the static mesh from the mesh list, and all its instances.
 	FOLIAGE_API void RemoveMesh(UFoliageType* InFoliageType);
