Fix for foliage UI continuing to display old foliage type's settings after changing StaticMesh on a foliage type.

diff --git a/Engine/Source/Editor/FoliageEdit/Private/FoliageEdMode.cpp b/Engine/Source/Editor/FoliageEdit/Private/FoliageEdMode.cpp
--- a/Engine/Source/Editor/FoliageEdit/Private/FoliageEdMode.cpp
+++ b/Engine/Source/Editor/FoliageEdit/Private/FoliageEdMode.cpp
@@ -166,7 +166,7 @@ void FEdModeFoliage::PostUndo()
 {
 	FEdMode::PostUndo();
 
-	StaticCastSharedPtr<FFoliageEdModeToolkit>(Toolkit)->PostUndo();
+	StaticCastSharedPtr<FFoliageEdModeToolkit>(Toolkit)->RefreshFullList();
 }
 
 /** When the user changes the active streaming level with the level browser */
@@ -1898,7 +1898,7 @@ bool FEdModeFoliage::ReplaceStaticMesh(UFoliageType* OldSettings, UStaticMesh* N
 
 			GEditor->BeginTransaction(NSLOCTEXT("UnrealEd", "FoliageMode_ChangeStaticMeshTransaction", "Foliage Editing: Change StaticMesh"));
 			IFA->Modify();
-			NewMeshInfo = IFA->AddMesh(NewStaticMesh, &NewSettings);
+			NewMeshInfo = IFA->AddMesh(NewStaticMesh, &NewSettings, Cast<UFoliageType_InstancedStaticMesh>(OldSettings));
 			NewSettings->DisplayOrder = OldSettings->DisplayOrder;
 			NewSettings->ShowNothing = OldSettings->ShowNothing;
 			NewSettings->ShowPaintSettings = OldSettings->ShowPaintSettings;
@@ -1935,6 +1935,8 @@ bool FEdModeFoliage::ReplaceStaticMesh(UFoliageType* OldSettings, UStaticMesh* N
 
 		// Update mesh list.
 		UpdateFoliageMeshList();
+
+		StaticCastSharedPtr<FFoliageEdModeToolkit>(Toolkit)->RefreshFullList();
 	}
 	return true;
 }
diff --git a/Engine/Source/Editor/FoliageEdit/Private/FoliageEdModeToolkit.cpp b/Engine/Source/Editor/FoliageEdit/Private/FoliageEdModeToolkit.cpp
--- a/Engine/Source/Editor/FoliageEdit/Private/FoliageEdModeToolkit.cpp
+++ b/Engine/Source/Editor/FoliageEdit/Private/FoliageEdModeToolkit.cpp
@@ -44,7 +44,7 @@ TSharedPtr<SWidget> FFoliageEdModeToolkit::GetInlineContent() const
 	return FoliageEdWidget;
 }
 
-void FFoliageEdModeToolkit::PostUndo()
+void FFoliageEdModeToolkit::RefreshFullList()
 {
 	// When an undo relates to the Foliage Edit mode, refresh the list.
 	FoliageEdWidget->RefreshFullList();
diff --git a/Engine/Source/Editor/FoliageEdit/Private/FoliageEdModeToolkit.h b/Engine/Source/Editor/FoliageEdit/Private/FoliageEdModeToolkit.h
--- a/Engine/Source/Editor/FoliageEdit/Private/FoliageEdModeToolkit.h
+++ b/Engine/Source/Editor/FoliageEdit/Private/FoliageEdModeToolkit.h
@@ -22,7 +22,7 @@ public:
 	virtual class FEdMode* GetEditorMode() const override;
 	virtual TSharedPtr<class SWidget> GetInlineContent() const override;
 
-	void PostUndo();
+	void RefreshFullList();
 
 private:
 	TSharedPtr< class SFoliageEdit > FoliageEdWidget;
