Fix foliage tool ApplyBrush doing mouse event things in tick without checking if it's the current viewport.
Should probably move all of this out of Tick and into a mouse event.

diff --git a/Engine/Source/Editor/FoliageEdit/Private/FoliageEdMode.cpp b/Engine/Source/Editor/FoliageEdit/Private/FoliageEdMode.cpp
--- a/Engine/Source/Editor/FoliageEdit/Private/FoliageEdMode.cpp
+++ b/Engine/Source/Editor/FoliageEdit/Private/FoliageEdMode.cpp
@@ -1198,7 +1198,7 @@ void FEdModeFoliage::ReapplyInstancesForBrush(UWorld* InWorld, AInstancedFoliage
 
 void FEdModeFoliage::ApplyBrush(FEditorViewportClient* ViewportClient)
 {
-	if (!bBrushTraceValid)
+	if (!bBrushTraceValid || ViewportClient != GCurrentLevelEditingViewportClient)
 	{
 		return;
 	}
