Seamless foliage painting/editing ( using shared foliage settings )
UE-8464

diff --git a/Engine/Source/Editor/FoliageEdit/Private/FoliageEdMode.cpp b/Engine/Source/Editor/FoliageEdit/Private/FoliageEdMode.cpp
--- a/Engine/Source/Editor/FoliageEdit/Private/FoliageEdMode.cpp
+++ b/Engine/Source/Editor/FoliageEdit/Private/FoliageEdMode.cpp
@@ -16,6 +16,7 @@
 //Slate dependencies
 #include "Editor/LevelEditor/Public/LevelEditor.h"
 #include "Editor/LevelEditor/Public/SLevelViewport.h"
+#include "Editor/UnrealEd/Public/Dialogs/DlgPickAssetPath.h"
 
 // Classes
 #include "InstancedFoliageActor.h"
@@ -37,6 +38,88 @@
 #define LOCTEXT_NAMESPACE "FoliageEdMode"
 #define FOLIAGE_SNAP_TRACE (10000.f)
 
+//
+// FFoliageMeshInfo iterator
+//
+class FFoliageMeshInfoIterator
+{
+private:
+	const UWorld*				World;
+	const UFoliageType*			FoliageType;
+	FFoliageMeshInfo*			CurrentMeshInfo;
+	AInstancedFoliageActor*		CurrentIFA;
+	int32						LevelIdx;
+
+public:
+	FFoliageMeshInfoIterator(const UWorld* InWorld, const UFoliageType* InFoliageType)
+		: World(InWorld)
+		, FoliageType(InFoliageType)
+		, CurrentMeshInfo(nullptr)
+		, CurrentIFA(nullptr)
+	{
+		// shortcut for non-assets
+		if (!InFoliageType->IsAsset())
+		{
+			LevelIdx = World->GetNumLevels();
+			AInstancedFoliageActor* IFA = Cast<AInstancedFoliageActor>(FoliageType->GetOuter());
+			if (IFA->GetLevel()->bIsVisible)
+			{
+				CurrentIFA = IFA;
+				CurrentMeshInfo = CurrentIFA->FindMesh(FoliageType);
+			}
+		}
+		else
+		{
+			LevelIdx= -1;
+			++(*this);
+		}
+	}
+
+	void operator++()
+	{
+		const int32 NumLevels = World->GetNumLevels();
+		int32 LocalLevelIdx = LevelIdx;
+		
+		while (++LocalLevelIdx < NumLevels)
+		{
+			ULevel* Level = World->GetLevel(LocalLevelIdx);
+			if (Level && Level->bIsVisible)
+			{
+				CurrentIFA = AInstancedFoliageActor::GetInstancedFoliageActorForLevel(Level);
+				if (CurrentIFA)
+				{
+					CurrentMeshInfo = CurrentIFA->FindMesh(FoliageType);
+					if (CurrentMeshInfo)
+					{
+						LevelIdx = LocalLevelIdx;
+						return;
+					}
+				}
+			}
+		}
+
+		CurrentMeshInfo = nullptr;
+		CurrentIFA = nullptr;
+	}
+
+	FORCEINLINE FFoliageMeshInfo* operator*()
+	{
+		check(CurrentMeshInfo);
+		return CurrentMeshInfo;
+	}
+	
+	FORCEINLINE operator bool() const
+	{
+		return CurrentMeshInfo != nullptr;
+	}
+
+	FORCEINLINE AInstancedFoliageActor* GetActor()
+	{
+		return CurrentIFA;
+	}
+};
+
+
 //
 // FEdModeFoliage
 //
@@ -46,7 +129,6 @@ FEdModeFoliage::FEdModeFoliage()
 	: FEdMode()
 	, bToolActive(false)
 	, bCanAltDrag(false)
-	, SelectionIFA(nullptr)
 {
 	// Load resources and construct brush component
 	UMaterial* BrushMaterial = nullptr;
@@ -101,6 +183,8 @@ void FEdModeFoliage::Enter()
 
 	// Bind to editor callbacks
 	FEditorDelegates::NewCurrentLevel.AddSP(this, &FEdModeFoliage::NotifyNewCurrentLevel);
+	FWorldDelegates::LevelAddedToWorld.AddSP(this, &FEdModeFoliage::NotifyLevelAddedToWorld);
+	FWorldDelegates::LevelRemovedFromWorld.AddSP(this, &FEdModeFoliage::NotifyLevelRemovedFromWorld);
 
 	// Force real-time viewports.  We'll back up the current viewport state so we can restore it when the
 	// user exits this mode.
@@ -108,22 +192,19 @@ void FEdModeFoliage::Enter()
 	const bool bRememberCurrentState = true;
 	ForceRealTimeViewports(bWantRealTime, bRememberCurrentState);
 
-	// Reapply selection visualization on any foliage items
-	if (UISettings.GetSelectToolSelected() || UISettings.GetLassoSelectToolSelected())
-	{
-		SelectionIFA = AInstancedFoliageActor::GetInstancedFoliageActorForCurrentLevel(GetWorld());
-		SelectionIFA->ApplySelectionToComponents(true);
-	}
-	else
-	{
-		SelectionIFA = nullptr;
-	}
-
 	if (!Toolkit.IsValid())
 	{
 		Toolkit = MakeShareable(new FFoliageEdModeToolkit);
 		Toolkit->Init(Owner->GetToolkitHost());
 	}
+	
+	if (UISettings.GetSelectToolSelected() || UISettings.GetLassoSelectToolSelected())
+	{
+		ApplySelectionToComponents(GetWorld(), true);
+	}
+		
+	// Update UI
+	NotifyNewCurrentLevel();
 }
 
 /** FEdMode: Called when the mode is exited */
@@ -134,7 +215,11 @@ void FEdModeFoliage::Exit()
 
 	//
 	FEditorDelegates::NewCurrentLevel.RemoveAll(this);
-
+	FWorldDelegates::LevelAddedToWorld.RemoveAll(this);
+	FWorldDelegates::LevelRemovedFromWorld.RemoveAll(this);
+	
+	FoliageMeshList.Empty();
+	
 	// Remove the brush
 	SphereBrushComponent->UnregisterComponent();
 
@@ -149,14 +234,8 @@ void FEdModeFoliage::Exit()
 	// Save UI settings to config file
 	UISettings.Save();
 
-	// Clear the placed level info
-	FoliageMeshList.Empty();
-
 	// Clear selection visualization on any foliage items
-	if (SelectionIFA && (UISettings.GetSelectToolSelected() || UISettings.GetLassoSelectToolSelected()))
-	{
-		SelectionIFA->ApplySelectionToComponents(false);
-	}
+	ApplySelectionToComponents(GetWorld(), false);
 
 	// Call base Exit method to ensure proper cleanup
 	FEdMode::Exit();
@@ -166,22 +245,23 @@ void FEdModeFoliage::PostUndo()
 {
 	FEdMode::PostUndo();
 
-	StaticCastSharedPtr<FFoliageEdModeToolkit>(Toolkit)->RefreshFullList();
+	PopulateFoliageMeshList();
 }
 
 /** When the user changes the active streaming level with the level browser */
 void FEdModeFoliage::NotifyNewCurrentLevel()
+{			
+	PopulateFoliageMeshList();
+}
+
+void FEdModeFoliage::NotifyLevelAddedToWorld(ULevel* InLevel, UWorld* InWorld)
 {
-	// Remove any selections in the old level and reapply for the new level
-	if (UISettings.GetSelectToolSelected() || UISettings.GetLassoSelectToolSelected())
-	{
-		if (SelectionIFA)
-		{
-			SelectionIFA->ApplySelectionToComponents(false);
-		}
-		SelectionIFA = AInstancedFoliageActor::GetInstancedFoliageActorForCurrentLevel(GetWorld());
-		SelectionIFA->ApplySelectionToComponents(true);
-	}
+	PopulateFoliageMeshList();
+}
+
+void FEdModeFoliage::NotifyLevelRemovedFromWorld(ULevel* InLevel, UWorld* InWorld)
+{
+	PopulateFoliageMeshList();
 }
 
 /** When the user changes the current tool in the UI */
@@ -189,19 +269,11 @@ void FEdModeFoliage::NotifyToolChanged()
 {
 	if (UISettings.GetSelectToolSelected() || UISettings.GetLassoSelectToolSelected())
 	{
-		if (SelectionIFA == nullptr)
-		{
-			SelectionIFA = AInstancedFoliageActor::GetInstancedFoliageActorForCurrentLevel(GetWorld());
-			SelectionIFA->ApplySelectionToComponents(true);
-		}
+		ApplySelectionToComponents(GetWorld(), true);
 	}
 	else
 	{
-		if (SelectionIFA)
-		{
-			SelectionIFA->ApplySelectionToComponents(false);
-		}
-		SelectionIFA = nullptr;
+		ApplySelectionToComponents(GetWorld(), false);
 	}
 }
 
@@ -224,8 +296,7 @@ void FEdModeFoliage::Tick(FEditorViewportClient* ViewportClient, float DeltaTime
 	if (UISettings.GetSelectToolSelected() || UISettings.GetLassoSelectToolSelected())
 	{
 		// Update pivot
-		AInstancedFoliageActor* IFA = AInstancedFoliageActor::GetInstancedFoliageActorForCurrentLevel(GetWorld());
-		Owner->PivotLocation = Owner->SnappedLocation = IFA->GetSelectionLocation();
+		UpdateWidgetLocationToInstanceSelection();
 	}
 
 	// Update the position and size of the brush component
@@ -273,7 +344,7 @@ void FEdModeFoliage::FoliageBrushTrace(FEditorViewportClient* ViewportClient, in
 			FHitResult Hit;
 			UWorld* World = ViewportClient->GetWorld();
 			static FName NAME_FoliageBrush = FName(TEXT("FoliageBrush"));
-			if (AInstancedFoliageActor::FoliageTrace(World, Hit, FDesiredFoliageInstance(Start, End), nullptr, NAME_FoliageBrush))
+			if (AInstancedFoliageActor::FoliageTrace(World, Hit, FDesiredFoliageInstance(Start, End), NAME_FoliageBrush))
 			{
 				// Check filters
 				UPrimitiveComponent* PrimComp = Hit.Component.Get();
@@ -365,16 +436,44 @@ static bool CheckLocationForPotentialInstance_ThreadSafe(const UFoliageType* Set
 	return true;
 }
 
-static bool CheckLocationForPotentialInstance(FFoliageMeshInfo& MeshInfo, const UFoliageType* Settings, float DensityCheckRadius, const FVector& Location, const FVector& Normal, TArray<FVector>& PotentialInstanceLocations, FFoliageInstanceHash& PotentialInstanceHash)
+static bool CheckForOverlappingSphere(AInstancedFoliageActor* IFA, const UFoliageType* Settings, const FSphere& Sphere)
+{
+	if (IFA)
+	{
+		FFoliageMeshInfo* MeshInfo = IFA->FindMesh(Settings);
+		if (MeshInfo)
+		{
+			return MeshInfo->CheckForOverlappingSphere(Sphere);
+		}
+	}
+	
+	return false;
+}
+
+// Returns whether or not there is are any instances overlapping the sphere specified
+static bool CheckForOverlappingSphere(const UWorld* InWorld, const UFoliageType* Settings, const FSphere& Sphere)
 {
+	for (FFoliageMeshInfoIterator It(InWorld, Settings); It; ++It)
+	{
+		FFoliageMeshInfo* MeshInfo = (*It);
+		if (MeshInfo->CheckForOverlappingSphere(Sphere))
+		{
+			return true;
+		}
+	}
+	
+	return false;
+}
 
+static bool CheckLocationForPotentialInstance(const UWorld* InWorld, const UFoliageType* Settings, float DensityCheckRadius, const FVector& Location, const FVector& Normal, TArray<FVector>& PotentialInstanceLocations, FFoliageInstanceHash& PotentialInstanceHash)
+{
 	if (CheckLocationForPotentialInstance_ThreadSafe(Settings, Location, Normal) == false)
 	{
 		return false;
 	}
 
 	// Check existing instances. Use the Density radius rather than the minimum radius
-	if (MeshInfo.CheckForOverlappingSphere(FSphere(Location, DensityCheckRadius)))
+	if (CheckForOverlappingSphere(InWorld, Settings, FSphere(Location, DensityCheckRadius)))
 	{
 		return false;
 	}
@@ -495,8 +594,7 @@ bool GeometryFilterCheck(const FHitResult& Hit, const UWorld* InWorld, const FDe
 		if (DesiredInst.PlacementMode == EFoliagePlacementMode::Manual)
 		{
 			UMaterialInterface* Material = PrimComp ? PrimComp->GetMaterial(0) : nullptr;
-			if (PrimComp->GetOutermost() != InWorld->GetCurrentLevel()->GetOutermost() ||
-				(!UISettings->bFilterLandscape && PrimComp->IsA(ULandscapeHeightfieldCollisionComponent::StaticClass())) ||
+			if ((!UISettings->bFilterLandscape && PrimComp->IsA(ULandscapeHeightfieldCollisionComponent::StaticClass())) ||
 				(!UISettings->bFilterStaticMesh && PrimComp->IsA(UStaticMeshComponent::StaticClass())) ||
 				(!UISettings->bFilterBSP && PrimComp->IsA(UModelComponent::StaticClass())) ||
 				(!UISettings->bFilterTranslucent && Material && IsTranslucentBlendMode(Material->GetBlendMode()))
@@ -549,7 +647,7 @@ bool LandscapeLayerCheck(const FHitResult& Hit, const UFoliageType* Settings, FE
 	return true;
 }
 
-void FEdModeFoliage::CalculatePotentialInstances_ThreadSafe(const UWorld* InWorld, const AInstancedFoliageActor* IgnoreIFA, const UFoliageType* Settings, const TArray<FDesiredFoliageInstance>* DesiredInstances, TArray<FPotentialInstance> OutPotentialInstances[NUM_INSTANCE_BUCKETS], const FFoliageUISettings* UISettings, const int32 StartIdx, const int32 LastIdx)
+void FEdModeFoliage::CalculatePotentialInstances_ThreadSafe(const UWorld* InWorld, const UFoliageType* Settings, const TArray<FDesiredFoliageInstance>* DesiredInstances, TArray<FPotentialInstance> OutPotentialInstances[NUM_INSTANCE_BUCKETS], const FFoliageUISettings* UISettings, const int32 StartIdx, const int32 LastIdx)
 {
 	LandscapeLayerCacheData LocalCache;
 
@@ -565,7 +663,7 @@ void FEdModeFoliage::CalculatePotentialInstances_ThreadSafe(const UWorld* InWorl
 		const FDesiredFoliageInstance& DesiredInst = (*DesiredInstances)[InstanceIdx];
 		FHitResult Hit;
 		static FName NAME_AddFoliageInstances = FName(TEXT("AddFoliageInstances"));
-		if (AInstancedFoliageActor::FoliageTrace(InWorld, Hit, DesiredInst, IgnoreIFA, NAME_AddFoliageInstances, true))
+		if (AInstancedFoliageActor::FoliageTrace(InWorld, Hit, DesiredInst, NAME_AddFoliageInstances, true))
 		{
 			float HitWeight = 1.f;
 			const bool bValidInstance = GeometryFilterCheck(Hit, InWorld, DesiredInst, UISettings)
@@ -573,18 +671,16 @@ void FEdModeFoliage::CalculatePotentialInstances_ThreadSafe(const UWorld* InWorl
 										&& VertexMaskCheck(Hit, Settings)
 										&& LandscapeLayerCheck(Hit, Settings, LocalCache, HitWeight);
 
-			
-
 			if (bValidInstance)
 			{
 				const int32 BucketIndex = FMath::RoundToInt(HitWeight * (float)(NUM_INSTANCE_BUCKETS - 1));
-				OutPotentialInstances[BucketIndex].Add(FPotentialInstance(Hit.ImpactPoint, Hit.ImpactNormal, Hit.Component.Get(), HitWeight, nullptr, nullptr, DesiredInst));
+				OutPotentialInstances[BucketIndex].Add(FPotentialInstance(Hit.ImpactPoint, Hit.ImpactNormal, Hit.Component.Get(), HitWeight, DesiredInst));
 			}
 		}
 	}
 }
 
-void FEdModeFoliage::CalculatePotentialInstances(const UWorld* InWorld, const AInstancedFoliageActor* IgnoreIFA, const UFoliageType* Settings, const TArray<FDesiredFoliageInstance>& DesiredInstances, TArray<FPotentialInstance> OutPotentialInstances[NUM_INSTANCE_BUCKETS], LandscapeLayerCacheData* LandscapeLayerCachesPtr, const FFoliageUISettings* UISettings)
+void FEdModeFoliage::CalculatePotentialInstances(const UWorld* InWorld, const UFoliageType* Settings, const TArray<FDesiredFoliageInstance>& DesiredInstances, TArray<FPotentialInstance> OutPotentialInstances[NUM_INSTANCE_BUCKETS], LandscapeLayerCacheData* LandscapeLayerCachesPtr, const FFoliageUISettings* UISettings)
 {
 	LandscapeLayerCacheData LocalCache;
 	LandscapeLayerCachesPtr = LandscapeLayerCachesPtr ? LandscapeLayerCachesPtr : &LocalCache;
@@ -608,60 +704,32 @@ void FEdModeFoliage::CalculatePotentialInstances(const UWorld* InWorld, const AI
 	{
 		FHitResult Hit;
 		static FName NAME_AddFoliageInstances = FName(TEXT("AddFoliageInstances"));
-		if (AInstancedFoliageActor::FoliageTrace(InWorld, Hit, DesiredInst, IgnoreIFA, NAME_AddFoliageInstances, true))
+		if (AInstancedFoliageActor::FoliageTrace(InWorld, Hit, DesiredInst, NAME_AddFoliageInstances, true))
 		{
 			float HitWeight = 1.f;
-			if (GeometryFilterCheck(Hit, InWorld, DesiredInst, UISettings) == false)
-			{
-				continue;
-			}
-			
+
 			UPrimitiveComponent* InstanceBase = Hit.GetComponent();
 			check(InstanceBase);
 
-			// New foliage instances always go into base component level
-			AInstancedFoliageActor* TargetIFA = AInstancedFoliageActor::GetInstancedFoliageActorForLevel(InstanceBase->GetComponentLevel(), true);
-			if (TargetIFA == nullptr)
-			{
-				continue;
-			}
-
+			ULevel* TargetLevel = InstanceBase->GetComponentLevel();
 			// We can paint into new level only if FoliageType is shared
-			if (DesiredInst.PlacementMode == EFoliagePlacementMode::Manual)
-			{
-				if (!Settings->IsAsset() && Settings->GetOuter() != TargetIFA)
+			if (!CanPaint(Settings, TargetLevel))
 			{
 				continue;
 			}
-			}
 
-			FFoliageMeshInfo* MeshInfo = DesiredInst.MeshInfo;
-			//We have to add the mesh into the foliage actor. This assumes the foliage type is just a concrete InstancedStaticMesh.
-			//With the new tree system coming this will not be true and you'll need to traverse the tree and insert the needed meshes
-			if (UStaticMesh* StaticMesh = Settings->GetStaticMesh())
+			if (!GeometryFilterCheck(Hit, InWorld, DesiredInst, UISettings))
 			{
-				UFoliageType* ExistingSettings = TargetIFA->GetSettingsForMesh(StaticMesh);
-				if (ExistingSettings == nullptr)
-				{
-					if (const UFoliageType_InstancedStaticMesh* InstancedStaticMeshFoliageType = Cast<UFoliageType_InstancedStaticMesh>(Settings))
-					{
-						MeshInfo = TargetIFA->AddMesh(Settings->GetStaticMesh(), nullptr, InstancedStaticMeshFoliageType);
-					}
-				}
-				else
-				{
-					MeshInfo = TargetIFA->FindMesh(ExistingSettings);
-				}
+				continue;
 			}
-
-			if (MeshInfo)
+			
+			const bool bValidInstance =	CheckLocationForPotentialInstance(InWorld, Settings, DensityCheckRadius, Hit.ImpactPoint, Hit.ImpactNormal, PotentialInstanceLocations, PotentialInstanceHash)
+										&& VertexMaskCheck(Hit, Settings)
+										&& LandscapeLayerCheck(Hit, Settings, LocalCache, HitWeight);
+			if (bValidInstance)
 			{
-				const bool bValidInstance =	CheckLocationForPotentialInstance(*MeshInfo, Settings, DensityCheckRadius, Hit.ImpactPoint, Hit.ImpactNormal, PotentialInstanceLocations, PotentialInstanceHash)
-											&& VertexMaskCheck(Hit, Settings)
-											&& LandscapeLayerCheck(Hit, Settings, LocalCache, HitWeight);
-
 				const int32 BucketIndex = FMath::RoundToInt(HitWeight * (float)(NUM_INSTANCE_BUCKETS - 1));
-				OutPotentialInstances[BucketIndex].Add(FPotentialInstance(Hit.ImpactPoint, Hit.ImpactNormal, InstanceBase, HitWeight, MeshInfo, TargetIFA, DesiredInst));
+				OutPotentialInstances[BucketIndex].Add(FPotentialInstance(Hit.ImpactPoint, Hit.ImpactNormal, InstanceBase, HitWeight, DesiredInst));
 			}
 		}
 	}
@@ -685,6 +753,18 @@ void FEdModeFoliage::AddInstances(UWorld* InWorld, const TArray<FDesiredFoliageI
 	}
 }
 
+static void SpawnFoliageInstance(UWorld* InWorld, const UFoliageType* Settings, const FFoliageInstance& Instance, UActorComponent* BaseComponent)
+{
+	// We always spawn instances in base component level
+	ULevel* TargetLevel = BaseComponent->GetComponentLevel();
+	AInstancedFoliageActor* IFA = AInstancedFoliageActor::GetInstancedFoliageActorForLevel(TargetLevel, true);
+	
+	FFoliageMeshInfo* MeshInfo;
+	UFoliageType* FoliageSettings = IFA->AddFoliageType(Settings, &MeshInfo);
+
+	MeshInfo->AddInstance(IFA, FoliageSettings, Instance, BaseComponent);
+}
+
 FFoliageMeshInfo* UpdateMeshSettings(AInstancedFoliageActor* IFA, const FDesiredFoliageInstance& DesiredInstance)
 {
 	FFoliageMeshInfo* MeshInfo = nullptr;
@@ -708,23 +788,22 @@ FFoliageMeshInfo* UpdateMeshSettings(AInstancedFoliageActor* IFA, const FDesired
 	return MeshInfo;
 }
 
-void FEdModeFoliage::AddInstancesImp(UWorld* InWorld, const UFoliageType* Settings, const TArray<FDesiredFoliageInstance>& DesiredInstances, const AInstancedFoliageActor* IgnoreIFA, const TArray<int32>& ExistingInstanceBuckets, const float Pressure, LandscapeLayerCacheData* LandscapeLayerCachesPtr, const FFoliageUISettings* UISettings)
+void FEdModeFoliage::AddInstancesImp(UWorld* InWorld, const UFoliageType* Settings, const TArray<FDesiredFoliageInstance>& DesiredInstances, const TArray<int32>& ExistingInstanceBuckets, const float Pressure, LandscapeLayerCacheData* LandscapeLayerCachesPtr, const FFoliageUISettings* UISettings)
 {
 	if (DesiredInstances.Num() == 0)
 	{
 		return;
 	}
 
-
 	TArray<FPotentialInstance> PotentialInstanceBuckets[NUM_INSTANCE_BUCKETS];
 	if (DesiredInstances[0].PlacementMode == EFoliagePlacementMode::Manual)
 	{
-	CalculatePotentialInstances(InWorld, IgnoreIFA, Settings, DesiredInstances, PotentialInstanceBuckets, LandscapeLayerCachesPtr, UISettings);
+		CalculatePotentialInstances(InWorld, Settings, DesiredInstances, PotentialInstanceBuckets, LandscapeLayerCachesPtr, UISettings);
 	}
 	else
 	{
 		//@TODO: actual threaded part coming, need parts of this refactor sooner for content team
-		CalculatePotentialInstances_ThreadSafe(InWorld, IgnoreIFA, Settings, &DesiredInstances, PotentialInstanceBuckets, nullptr, 0, DesiredInstances.Num() - 1);
+		CalculatePotentialInstances_ThreadSafe(InWorld, Settings, &DesiredInstances, PotentialInstanceBuckets, nullptr, 0, DesiredInstances.Num() - 1);
 
 		TMap<AInstancedFoliageActor*, FFoliageMeshInfo*> UpdatedIFAs;	//we want to override any existing mesh settings with the procedural settings.
 		for (auto& Bucket : PotentialInstanceBuckets)
@@ -732,7 +811,7 @@ void FEdModeFoliage::AddInstancesImp(UWorld* InWorld, const UFoliageType* Settin
 			for (FPotentialInstance& PotentialInst : Bucket)
 			{
 				// New foliage instances always go into base component level
-				AInstancedFoliageActor* TargetIFA = AInstancedFoliageActor::GetInstancedFoliageActorForLevel(PotentialInst.HitComponent->GetComponentLevel());
+				AInstancedFoliageActor* TargetIFA = AInstancedFoliageActor::GetInstancedFoliageActorForLevel(PotentialInst.HitComponent->GetComponentLevel(), true);
 				FFoliageMeshInfo* MeshInfo = nullptr;
 				if (FFoliageMeshInfo** MeshInfoResult = UpdatedIFAs.Find(TargetIFA))
 				{
@@ -743,15 +822,9 @@ void FEdModeFoliage::AddInstancesImp(UWorld* InWorld, const UFoliageType* Settin
 					MeshInfo = UpdateMeshSettings(TargetIFA, PotentialInst.DesiredInstance);
 					UpdatedIFAs.Add(TargetIFA, MeshInfo);
 				}
-
-				PotentialInst.IFA = TargetIFA;
-				PotentialInst.MeshInfo = MeshInfo;
 			}
 		}
 	}
-		
-	UStaticMesh* StaticMesh = Settings->GetStaticMesh();
-	
 
 	for (int32 BucketIdx = 0; BucketIdx < NUM_INSTANCE_BUCKETS; BucketIdx++)
 	{
@@ -766,32 +839,40 @@ void FEdModeFoliage::AddInstancesImp(UWorld* InWorld, const UFoliageType* Settin
 		{
 			FPotentialInstance& PotentialInstance = PotentialInstances[Idx];
 			FFoliageInstance Inst;
-			if (PotentialInstance.PlaceInstance(Settings, Inst))
+			if (PotentialInstance.PlaceInstance(InWorld, Settings, Inst))
 			{
 				Inst.ProceduralGuid = PotentialInstance.DesiredInstance.ProceduralGuid;
-				PotentialInstance.MeshInfo->AddInstance(PotentialInstance.IFA, Settings, Inst);
+				
+				SpawnFoliageInstance(InWorld, Settings, Inst, PotentialInstance.HitComponent);
 			}
 		}
 	}
 }
 
 /** Add instances inside the brush to match DesiredInstanceCount */
-void FEdModeFoliage::AddInstancesForBrush(UWorld* InWorld, AInstancedFoliageActor* IFA, UFoliageType* Settings, FFoliageMeshInfo& MeshInfo, int32 DesiredInstanceCount, const TArray<int32>& ExistingInstances, float Pressure)
+void FEdModeFoliage::AddInstancesForBrush(UWorld* InWorld, const UFoliageType* Settings, const FSphere& BrushSphere, int32 DesiredInstanceCount, float Pressure)
 {
-	checkf(InWorld == IFA->GetWorld(), TEXT("Warning:World does not match Foliage world"));
-	if (DesiredInstanceCount > ExistingInstances.Num())
-	{
-		TArray<int32> ExistingInstanceBuckets;
-		ExistingInstanceBuckets.AddZeroed(NUM_INSTANCE_BUCKETS);
+	UWorld* World = GetWorld();
+	const bool bHasValidLandscapeLayers = LandscapeLayersValid(Settings);
+		
+	TArray<int32> ExistingInstanceBuckets;
+	ExistingInstanceBuckets.AddZeroed(NUM_INSTANCE_BUCKETS);
+	int32 NumExistingInstances = 0;
 
-		const TArray<FName>& LandscapeLayers = Settings->LandscapeLayers;
-		if (LandscapeLayersValid(Settings))
+	for (FFoliageMeshInfoIterator It(World, Settings); It; ++It)
+	{
+		FFoliageMeshInfo* MeshInfo = (*It);
+		TArray<int32> ExistingInstances;
+		MeshInfo->GetInstancesInsideSphere(BrushSphere, ExistingInstances);
+		NumExistingInstances+= ExistingInstances.Num();
+		
+		if (bHasValidLandscapeLayers)
 		{
 			// Find the landscape weights of existing ExistingInstances
 			for (int32 Idx : ExistingInstances)
 			{
-				FFoliageInstance& Instance = MeshInfo.Instances[Idx];
-				auto InstanceBasePtr = IFA->InstanceBaseCache.GetInstanceBasePtr(Instance.BaseId);
+				FFoliageInstance& Instance = MeshInfo->Instances[Idx];
+				auto InstanceBasePtr = It.GetActor()->InstanceBaseCache.GetInstanceBasePtr(Instance.BaseId);
 				float HitWeight;
 				if (GetMaxHitWeight(Instance.Location, InstanceBasePtr.Get(), Settings, &LandscapeLayerCaches, HitWeight))
 				{
@@ -803,10 +884,12 @@ void FEdModeFoliage::AddInstancesForBrush(UWorld* InWorld, AInstancedFoliageActo
 		else
 		{
 			// When not tied to a layer, put all the ExistingInstances in the last bucket.
-			ExistingInstanceBuckets[NUM_INSTANCE_BUCKETS - 1] = ExistingInstances.Num();
+			ExistingInstanceBuckets[NUM_INSTANCE_BUCKETS - 1] = NumExistingInstances;
 		}
+	}
 
-		TArray<FPotentialInstance> PotentialInstanceBuckets[NUM_INSTANCE_BUCKETS];
+	if (DesiredInstanceCount > NumExistingInstances)
+	{
 		TArray<FDesiredFoliageInstance> DesiredInstances;	//we compute instances for the brush
 		DesiredInstances.Reserve(DesiredInstanceCount);
 
@@ -815,67 +898,280 @@ void FEdModeFoliage::AddInstancesForBrush(UWorld* InWorld, AInstancedFoliageActo
 			FVector Start, End;
 			GetRandomVectorInBrush(Start, End);
 			FDesiredFoliageInstance* DesiredInstance = new (DesiredInstances)FDesiredFoliageInstance(Start, End);
-			DesiredInstance->MeshInfo = &MeshInfo;
 		}
 
-		AddInstancesImp(InWorld, Settings, DesiredInstances, IFA, ExistingInstanceBuckets, Pressure, &LandscapeLayerCaches, &UISettings);
+		AddInstancesImp(InWorld, Settings, DesiredInstances, ExistingInstanceBuckets, Pressure, &LandscapeLayerCaches, &UISettings);
 	}
 }
 
 /** Remove instances inside the brush to match DesiredInstanceCount */
-void FEdModeFoliage::RemoveInstancesForBrush(AInstancedFoliageActor* IFA, FFoliageMeshInfo& MeshInfo, int32 DesiredInstanceCount, TArray<int32>& PotentialInstancesToRemove, float Pressure)
+void FEdModeFoliage::RemoveInstancesForBrush(UWorld* InWorld, const UFoliageType* Settings, const FSphere& BrushSphere, int32 DesiredInstanceCount, float Pressure)
+{
+	for (FFoliageMeshInfoIterator It(InWorld, Settings); It; ++It)
+	{
+		FFoliageMeshInfo* MeshInfo = (*It);
+		AInstancedFoliageActor* IFA = It.GetActor();
+
+		TArray<int32> PotentialInstancesToRemove;
+		MeshInfo->GetInstancesInsideSphere(BrushSphere, PotentialInstancesToRemove);
+		if (PotentialInstancesToRemove.Num() == 0)
+		{
+			continue;
+		}
+		
+		int32 InstancesToRemove = FMath::RoundToInt((float)(PotentialInstancesToRemove.Num() - DesiredInstanceCount) * Pressure);
+		int32 InstancesToKeep = PotentialInstancesToRemove.Num() - InstancesToRemove;
+		if (InstancesToKeep > 0)
+		{
+			// Remove InstancesToKeep random PotentialInstancesToRemove from the array to leave those PotentialInstancesToRemove behind, and delete all the rest
+			for (int32 i = 0; i < InstancesToKeep; i++)
+			{
+				PotentialInstancesToRemove.RemoveAtSwap(FMath::Rand() % PotentialInstancesToRemove.Num());
+			}
+		}
+
+		if (!UISettings.bFilterLandscape || !UISettings.bFilterStaticMesh || !UISettings.bFilterBSP || !UISettings.bFilterTranslucent)
+		{
+			// Filter PotentialInstancesToRemove
+			for (int32 Idx = 0; Idx < PotentialInstancesToRemove.Num(); Idx++)
+			{
+				auto BaseId = MeshInfo->Instances[PotentialInstancesToRemove[Idx]].BaseId;
+				auto BasePtr = IFA->InstanceBaseCache.GetInstanceBasePtr(BaseId);
+				UPrimitiveComponent* Base = Cast<UPrimitiveComponent>(BasePtr.Get());
+				UMaterialInterface* Material = Base ? Base->GetMaterial(0) : nullptr;
+
+				// Check if instance is candidate for removal based on filter settings
+				if (Base && (
+					(!UISettings.bFilterLandscape && Base->IsA(ULandscapeHeightfieldCollisionComponent::StaticClass())) ||
+					(!UISettings.bFilterStaticMesh && Base->IsA(UStaticMeshComponent::StaticClass())) ||
+					(!UISettings.bFilterBSP && Base->IsA(UModelComponent::StaticClass())) ||
+					(!UISettings.bFilterTranslucent && Material && IsTranslucentBlendMode(Material->GetBlendMode()))
+					))
+				{
+					// Instance should not be removed, so remove it from the removal list.
+					PotentialInstancesToRemove.RemoveAtSwap(Idx);
+					Idx--;
+				}
+			}
+		}
+
+		// Remove PotentialInstancesToRemove to reduce it to desired count
+		if (PotentialInstancesToRemove.Num() > 0)
+		{
+			MeshInfo->RemoveInstances(IFA, PotentialInstancesToRemove);
+		}
+	}
+}
+
+void FEdModeFoliage::SelectInstancesForBrush(UWorld* InWorld, const UFoliageType* Settings, const FSphere& BrushSphere, bool bSelect)
+{
+	for (FFoliageMeshInfoIterator It(InWorld, Settings); It; ++It)
+	{
+		FFoliageMeshInfo* MeshInfo = (*It);
+		AInstancedFoliageActor* IFA = It.GetActor();
+
+		TArray<int32> Instances;
+		MeshInfo->GetInstancesInsideSphere(BrushSphere, Instances);
+		if (Instances.Num() == 0)
+		{
+			continue;
+		}
+
+		MeshInfo->SelectInstances(IFA, bSelect, Instances);
+	}
+}
+
+void FEdModeFoliage::SelectInstances(UWorld* InWorld, bool bSelect)
+{
+	for (auto& FoliageMeshUI : FoliageMeshList)
+	{	
+		UFoliageType* Settings = FoliageMeshUI->Settings;	
+		
+		if (bSelect && !Settings->IsSelected)
+		{
+			continue;
+		}
+
+		SelectInstances(InWorld, Settings, bSelect);
+	}
+}
+	
+void FEdModeFoliage::SelectInstances(UWorld* InWorld, const UFoliageType* Settings, bool bSelect)
+{
+	for (FFoliageMeshInfoIterator It(InWorld, Settings); It; ++It)
+	{
+		FFoliageMeshInfo* MeshInfo = (*It);
+		AInstancedFoliageActor* IFA = It.GetActor();
+
+		MeshInfo->SelectInstances(IFA, bSelect);
+	}
+}
+
+void FEdModeFoliage::ApplySelectionToComponents(UWorld* InWorld, bool bApply)
+{
+	const int32 NumLevels = InWorld->GetNumLevels();
+	for (int32 LevelIdx = 0; LevelIdx < NumLevels; ++LevelIdx)
+	{
+		ULevel* Level = InWorld->GetLevel(LevelIdx);
+		AInstancedFoliageActor* IFA = AInstancedFoliageActor::GetInstancedFoliageActorForLevel(Level);
+		if (IFA)
+		{
+			IFA->ApplySelectionToComponents(bApply);
+		}
+	}
+}
+
+void FEdModeFoliage::TransformSelectedInstances(UWorld* InWorld, const FVector& InDrag, const FRotator& InRot, const FVector& InScale, bool bDuplicate)
+{
+	const int32 NumLevels = InWorld->GetNumLevels();
+	for (int32 LevelIdx = 0; LevelIdx < NumLevels; ++LevelIdx)
+	{
+		ULevel* Level = InWorld->GetLevel(LevelIdx);
+		AInstancedFoliageActor* IFA = AInstancedFoliageActor::GetInstancedFoliageActorForLevel(Level);
+		if (IFA)
+		{
+			bool bFoundSelection = false;
+
+			for (auto& MeshPair : IFA->FoliageMeshes)
+			{
+				FFoliageMeshInfo& MeshInfo = *MeshPair.Value;
+				TArray<int32> SelectedIndices = MeshInfo.SelectedIndices.Array();
+
+				if (SelectedIndices.Num() > 0)
+				{
+					// Mark actor once we found selection
+					if (!bFoundSelection)
+					{
+						IFA->Modify();
+						bFoundSelection = true;
+					}
+					
+					if (bDuplicate)
+					{
+						MeshInfo.DuplicateInstances(IFA, MeshPair.Key, SelectedIndices);
+					}
+
+					MeshInfo.PreMoveInstances(IFA, SelectedIndices);
+
+					for (int32 SelectedInstanceIdx : SelectedIndices)
+					{
+						FFoliageInstance& Instance = MeshInfo.Instances[SelectedInstanceIdx];
+						Instance.Location += InDrag;
+						Instance.ZOffset = 0.f;
+						Instance.Rotation += InRot;
+						Instance.DrawScale3D += InScale;
+					}
+
+					MeshInfo.PostMoveInstances(IFA, SelectedIndices);
+				}
+			}
+
+			if (bFoundSelection)
+			{
+				IFA->MarkComponentsRenderStateDirty();
+			}
+		}
+	}
+}
+
+AInstancedFoliageActor* FEdModeFoliage::GetSelectionLocation(UWorld* InWorld, FVector& OutLocation) const
 {
-	int32 InstancesToRemove = FMath::RoundToInt((float)(PotentialInstancesToRemove.Num() - DesiredInstanceCount) * Pressure);
-	int32 InstancesToKeep = PotentialInstancesToRemove.Num() - InstancesToRemove;
-	if (InstancesToKeep > 0)
+	// Prefer current level
+	AInstancedFoliageActor* IFA = AInstancedFoliageActor::GetInstancedFoliageActorForCurrentLevel(InWorld);
+	if (IFA && IFA->GetSelectionLocation(OutLocation))
+	{
+		return IFA;
+	}
+	
+	// Go through all sub-levels
+	const int32 NumLevels = InWorld->GetNumLevels();
+	for (int32 LevelIdx = 0; LevelIdx < NumLevels; ++LevelIdx)
 	{
-		// Remove InstancesToKeep random PotentialInstancesToRemove from the array to leave those PotentialInstancesToRemove behind, and delete all the rest
-		for (int32 i = 0; i < InstancesToKeep; i++)
+		ULevel* Level = InWorld->GetLevel(LevelIdx);
+		if (Level != InWorld->GetCurrentLevel())
 		{
-			PotentialInstancesToRemove.RemoveAtSwap(FMath::Rand() % PotentialInstancesToRemove.Num());
+			AInstancedFoliageActor* IFA = AInstancedFoliageActor::GetInstancedFoliageActorForLevel(Level);
+			if (IFA && IFA->GetSelectionLocation(OutLocation))
+			{
+				return IFA;
+			}
 		}
 	}
+	
+	return nullptr;
+}
+
+void FEdModeFoliage::UpdateWidgetLocationToInstanceSelection()
+{
+	FVector SelectionLocation = FVector::ZeroVector;
+	AInstancedFoliageActor* IFA = GetSelectionLocation(GetWorld(), SelectionLocation);
+	Owner->PivotLocation = Owner->SnappedLocation = SelectionLocation;
+	//if (IFA)
+	//{
+	//	IFA->MarkComponentsRenderStateDirty();
+	//}
+}
 
-	if (!UISettings.bFilterLandscape || !UISettings.bFilterStaticMesh || !UISettings.bFilterBSP || !UISettings.bFilterTranslucent)
+void FEdModeFoliage::RemoveSelectedInstances(UWorld* InWorld)
+{
+	GEditor->BeginTransaction(NSLOCTEXT("UnrealEd", "FoliageMode_EditTransaction", "Foliage Editing"));
+	
+	const int32 NumLevels = InWorld->GetNumLevels();
+	for (int32 LevelIdx = 0; LevelIdx < NumLevels; ++LevelIdx)
 	{
-		// Filter PotentialInstancesToRemove
-		for (int32 Idx = 0; Idx < PotentialInstancesToRemove.Num(); Idx++)
+		ULevel* Level = InWorld->GetLevel(LevelIdx);
+		AInstancedFoliageActor* IFA = AInstancedFoliageActor::GetInstancedFoliageActorForLevel(Level);
+		if (IFA)
 		{
-			auto BasePtr = IFA->InstanceBaseCache.GetInstanceBasePtr(MeshInfo.Instances[PotentialInstancesToRemove[Idx]].BaseId);
-			UPrimitiveComponent* Base = Cast<UPrimitiveComponent>(BasePtr.Get());
-			UMaterialInterface* Material = Base ? Base->GetMaterial(0) : nullptr;
+			bool bHasSelection = false;
+			for (auto& MeshPair : IFA->FoliageMeshes)
+			{
+				if (MeshPair.Value->SelectedIndices.Num())
+				{
+					bHasSelection = true;
+					break;
+				}
+			}
 
-			// Check if instance is candidate for removal based on filter settings
-			if (Base && (
-				(!UISettings.bFilterLandscape && Base->IsA(ULandscapeHeightfieldCollisionComponent::StaticClass())) ||
-				(!UISettings.bFilterStaticMesh && Base->IsA(UStaticMeshComponent::StaticClass())) ||
-				(!UISettings.bFilterBSP && Base->IsA(UModelComponent::StaticClass())) ||
-				(!UISettings.bFilterTranslucent && Material && IsTranslucentBlendMode(Material->GetBlendMode()))
-				))
+			if (bHasSelection)
 			{
-				// Instance should not be removed, so remove it from the removal list.
-				PotentialInstancesToRemove.RemoveAtSwap(Idx);
-				Idx--;
+				IFA->Modify();
+				for (auto& MeshPair : IFA->FoliageMeshes)
+				{
+					FFoliageMeshInfo& Mesh = *MeshPair.Value;
+					if (Mesh.SelectedIndices.Num() > 0)
+					{
+						TArray<int32> InstancesToDelete = Mesh.SelectedIndices.Array();
+						Mesh.RemoveInstances(IFA, InstancesToDelete);
+					
+						OnInstanceCountUpdated(MeshPair.Key);
+					}
+				}
 			}
 		}
 	}
 
-	// Remove PotentialInstancesToRemove to reduce it to desired count
-	if (PotentialInstancesToRemove.Num() > 0)
+	GEditor->EndTransaction();
+}
+
+void FEdModeFoliage::ReapplyInstancesForBrush(UWorld* InWorld, const UFoliageType* Settings, const FSphere& BrushSphere)
+{
+	for (FFoliageMeshInfoIterator It(InWorld, Settings); It; ++It)
 	{
-		MeshInfo.RemoveInstances(IFA, PotentialInstancesToRemove);
+		FFoliageMeshInfo* MeshInfo = (*It);
+		AInstancedFoliageActor* IFA = It.GetActor();
+
+		ReapplyInstancesForBrush(InWorld, IFA, Settings, MeshInfo, BrushSphere);
 	}
 }
 
 /** Reapply instance settings to exiting instances */
-void FEdModeFoliage::ReapplyInstancesForBrush(UWorld* InWorld, AInstancedFoliageActor* IFA, UFoliageType* Settings, const TArray<int32>& ExistingInstances)
+void FEdModeFoliage::ReapplyInstancesForBrush(UWorld* InWorld, AInstancedFoliageActor* IFA, const UFoliageType* Settings, FFoliageMeshInfo* MeshInfo, const FSphere& BrushSphere)
 {
-	checkf(InWorld == IFA->GetWorld(), TEXT("Warning:World does not match Foliage world"));
-	FFoliageMeshInfo* MeshInfo = IFA->FindMesh(Settings);
-	check(MeshInfo);
-
+	TArray<int32> ExistingInstances;
+	MeshInfo->GetInstancesInsideSphere(BrushSphere, ExistingInstances);
+		
 	bool bUpdated = false;
-
 	TArray<int32> UpdatedInstances;
 	TSet<int32> InstancesToDelete;
 
@@ -1045,7 +1341,7 @@ void FEdModeFoliage::ReapplyInstancesForBrush(UWorld* InWorld, AInstancedFoliage
 				FVector ZAxis = Instance.Rotation.Quaternion().GetAxisZ();
 				FVector Start = Instance.Location + 16.f * ZAxis;
 				FVector End = Instance.Location - 16.f * ZAxis;
-				if (AInstancedFoliageActor::FoliageTrace(InWorld, Hit, FDesiredFoliageInstance(Start, End), IFA, NAME_ReapplyInstancesForBrush, true))
+				if (AInstancedFoliageActor::FoliageTrace(InWorld, Hit, FDesiredFoliageInstance(Start, End), NAME_ReapplyInstancesForBrush, true))
 				{
 					// Reapply the normal
 					if (bReapplyNormal)
@@ -1163,9 +1459,9 @@ void FEdModeFoliage::ReapplyInstancesForBrush(UWorld* InWorld, AInstancedFoliage
 				static const FName NAME_ReapplyInstancesForBrush = TEXT("ReapplyCollisionWithWorld");
 				FVector Start = Instance.Location + FVector(0.f, 0.f, 16.f);
 				FVector End = Instance.Location - FVector(0.f, 0.f, 16.f);
-				if (AInstancedFoliageActor::FoliageTrace(InWorld, Hit, FDesiredFoliageInstance(Start, End), IFA, NAME_ReapplyInstancesForBrush))
+				if (AInstancedFoliageActor::FoliageTrace(InWorld, Hit, FDesiredFoliageInstance(Start, End), NAME_ReapplyInstancesForBrush))
 				{
-					if (!IFA->CheckCollisionWithWorld(Settings, Instance, Hit.Normal, Hit.Location))
+					if (!AInstancedFoliageActor::CheckCollisionWithWorld(InWorld, Settings, Instance, Hit.Normal, Hit.Location))
 					{
 						InstancesToDelete.Add(InstanceIndex);
 						continue;
@@ -1196,95 +1492,132 @@ void FEdModeFoliage::ReapplyInstancesForBrush(UWorld* InWorld, AInstancedFoliage
 	}
 }
 
-void FEdModeFoliage::ApplyBrush(FEditorViewportClient* ViewportClient)
+void FEdModeFoliage::ReapplyInstancesDensityForBrush(UWorld* InWorld, const UFoliageType* Settings, const FSphere& BrushSphere, float Pressure)
 {
-	if (!bBrushTraceValid || ViewportClient != GCurrentLevelEditingViewportClient)
-	{
-		return;
-	}
-
-	AInstancedFoliageActor* IFA = AInstancedFoliageActor::GetInstancedFoliageActorForCurrentLevel(GetWorld());
-
-	float BrushArea = PI * FMath::Square(UISettings.GetRadius());
-
+	// TODO: make it work with seamless painting
+	
+	//// Adjust instance density
+	//int32 SnapShotInstanceCount = 0;
+	//TArray<FMeshInfoSnapshot*> SnapShotList;
+	//InstanceSnapshot.MultiFindPointer(Settings, SnapShotList);
+
+	//if (SnapShotList.Num())
+	//{
+	//	// Use snapshot to determine number of instances at the start of the brush stroke
+	//	for (auto* SnapShot : SnapShotList)
+	//	{
+	//		SnapShotInstanceCount+= SnapShot->CountInstancesInsideSphere(BrushSphere);
+	//	}
+	//			
+	//	int32 NewInstanceCount = FMath::RoundToInt((float)SnapShotInstanceCount * Settings->ReapplyDensityAmount);
+	//	if (Settings->ReapplyDensityAmount > 1.f && NewInstanceCount > Instances.Num())
+	//	{
+	//		AddInstancesForBrush(InWorld, Settings, BrushSphere, );
+	//	}
+	//	else if (Settings->ReapplyDensityAmount < 1.f && NewInstanceCount < Instances.Num())
+	//	{
+	//		RemoveInstancesForBrush(IFA, *MeshInfo, NewInstanceCount, Instances, Pressure);
+	//	}
+	//}
+}
+
+void FEdModeFoliage::PreApplyBrush()
+{
+	InstanceSnapshot.Empty();
+	
+	UWorld* World = GetWorld();
+	// Special setup beginning a stroke with the Reapply tool
+	// Necessary so we don't keep reapplying settings over and over for the same instances.
+	if (UISettings.GetReapplyToolSelected())
+	{
+		for (auto& FoliageMeshUI : FoliageMeshList)
+		{	
+			UFoliageType* Settings = FoliageMeshUI->Settings;	
+		
+			if (!Settings->IsSelected)
+			{
+				continue;
+			}
+
+			for (FFoliageMeshInfoIterator It(World, Settings); It; ++It)
+			{
+				FFoliageMeshInfo* MeshInfo = (*It);
+				
+				// Take a snapshot of all the locations
+				InstanceSnapshot.Add(Settings, FMeshInfoSnapshot(MeshInfo));
+
+				// Clear the "FOLIAGE_Readjusted" flag
+				for (auto& Instance : MeshInfo->Instances)
+				{
+					Instance.Flags &= (~FOLIAGE_Readjusted);
+				}
+			}
+		}
+	}
+}
+
+void FEdModeFoliage::ApplyBrush(FEditorViewportClient* ViewportClient)
+{
+	if (!bBrushTraceValid || ViewportClient != GCurrentLevelEditingViewportClient)
+	{
+		return;
+	}
+	
+	float BrushArea = PI * FMath::Square(UISettings.GetRadius());
+
 	// Tablet pressure
 	float Pressure = ViewportClient->Viewport->IsPenActive() ? ViewportClient->Viewport->GetTabletPressure() : 1.f;
 
 	// Cache a copy of the world pointer
 	UWorld* World = ViewportClient->GetWorld();
-
-	for (auto& MeshPair : IFA->FoliageMeshes)
-	{
-		FFoliageMeshInfo& MeshInfo = *MeshPair.Value;
-		UFoliageType* Settings = MeshPair.Key;
-
-		if (Settings->IsSelected)
+		
+	for (auto& FoliageMeshUI : FoliageMeshList)
+	{	
+		UFoliageType* Settings = FoliageMeshUI->Settings;	
+		
+		if (!Settings->IsSelected)
 		{
-			// Find the instances already in the area.
-			TArray<int32> Instances;
-			FSphere BrushSphere(BrushLocation, UISettings.GetRadius());
-			MeshInfo.GetInstancesInsideSphere(BrushSphere, Instances);
+			continue;
+		}
 
-			if (UISettings.GetLassoSelectToolSelected())
+		FSphere BrushSphere(BrushLocation, UISettings.GetRadius());
+					
+		if (UISettings.GetLassoSelectToolSelected())
+		{
+			// Shift unpaints
+			SelectInstancesForBrush(World, Settings, BrushSphere, !IsShiftDown(ViewportClient->Viewport));
+		}
+		else if (UISettings.GetReapplyToolSelected())
+		{
+			// Reapply any settings checked by the user
+			ReapplyInstancesForBrush(World, Settings, BrushSphere);
+		}
+		else if (UISettings.GetPaintToolSelected())
+		{
+			// Shift unpaints
+			if (IsShiftDown(ViewportClient->Viewport))
 			{
-				// Shift unpaints
-				MeshInfo.SelectInstances(IFA, !IsShiftDown(ViewportClient->Viewport), Instances);
+				int32 DesiredInstanceCount = FMath::RoundToInt(BrushArea * Settings->Density * UISettings.GetUnpaintDensity() / (1000.f*1000.f));
+					
+				RemoveInstancesForBrush(World, Settings, BrushSphere, DesiredInstanceCount, Pressure);
 			}
 			else
-			if (UISettings.GetReapplyToolSelected())
 			{
-				if (Settings->ReapplyDensity)
-				{
-					// Adjust instance density
-					FMeshInfoSnapshot* SnapShot = InstanceSnapshot.Find(MeshPair.Key);
-					if (SnapShot)
-					{
-						// Use snapshot to determine number of instances at the start of the brush stroke
-						int32 NewInstanceCount = FMath::RoundToInt((float)SnapShot->CountInstancesInsideSphere(BrushSphere) * Settings->ReapplyDensityAmount);
-						if (Settings->ReapplyDensityAmount > 1.f && NewInstanceCount > Instances.Num())
-						{
-							AddInstancesForBrush(World, IFA, Settings, MeshInfo, NewInstanceCount, Instances, Pressure);
-						}
-						else if (Settings->ReapplyDensityAmount < 1.f && NewInstanceCount < Instances.Num())
-						{
-							RemoveInstancesForBrush(IFA, MeshInfo, NewInstanceCount, Instances, Pressure);
-						}
-					}
-				}
-
-				// Reapply any settings checked by the user
-				ReapplyInstancesForBrush(World, IFA, Settings, Instances);
-			}
-			else if (UISettings.GetPaintToolSelected())
-			{
-				// Shift unpaints
-				if (IsShiftDown(ViewportClient->Viewport))
-				{
-					int32 DesiredInstanceCount = FMath::RoundToInt(BrushArea * Settings->Density * UISettings.GetUnpaintDensity() / (1000.f*1000.f));
-
-					if (DesiredInstanceCount < Instances.Num())
-					{
-						RemoveInstancesForBrush(IFA, MeshInfo, DesiredInstanceCount, Instances, Pressure);
-					}
-				}
-				else
-				{
-					// This is the total set of instances disregarding parameters like slope, height or layer.
-					float DesiredInstanceCountFloat = BrushArea * Settings->Density * UISettings.GetPaintDensity() / (1000.f*1000.f);
-
-					// Allow a single instance with a random chance, if the brush is smaller than the density
-					int32 DesiredInstanceCount = DesiredInstanceCountFloat > 1.f ? FMath::RoundToInt(DesiredInstanceCountFloat) : FMath::FRand() < DesiredInstanceCountFloat ? 1 : 0;
-
-					AddInstancesForBrush(World, IFA, Settings, MeshInfo, DesiredInstanceCount, Instances, Pressure);
-				}
+				// This is the total set of instances disregarding parameters like slope, height or layer.
+				float DesiredInstanceCountFloat = BrushArea * Settings->Density * UISettings.GetPaintDensity() / (1000.f*1000.f);
+				// Allow a single instance with a random chance, if the brush is smaller than the density
+				int32 DesiredInstanceCount = DesiredInstanceCountFloat > 1.f ? FMath::RoundToInt(DesiredInstanceCountFloat) : FMath::FRand() < DesiredInstanceCountFloat ? 1 : 0;
+					
+				AddInstancesForBrush(World, Settings, BrushSphere, DesiredInstanceCount, Pressure);
 			}
 		}
+		
+		OnInstanceCountUpdated(Settings);
 	}
+
 	if (UISettings.GetLassoSelectToolSelected())
 	{
-		IFA->CheckSelection();
-		Owner->PivotLocation = Owner->SnappedLocation = IFA->GetSelectionLocation();
-		IFA->MarkComponentsRenderStateDirty();
+		UpdateWidgetLocationToInstanceSelection();
 	}
 }
 
@@ -1334,199 +1667,190 @@ struct FFoliagePaintBucketTriangle
 };
 
 /** Apply paint bucket to actor */
-void FEdModeFoliage::ApplyPaintBucket(AActor* Actor, bool bRemove)
+void FEdModeFoliage::ApplyPaintBucket_Remove(AActor* Actor)
 {
-	// Apply only to current level
-	if (!Actor->GetLevel()->IsCurrentLevel())
-	{
-		return;
-	}
+	UWorld* World = Actor->GetWorld();
 
-	if (bRemove)
-	{
-		// Remove all instances of the selected meshes
-		AInstancedFoliageActor* IFA = AInstancedFoliageActor::GetInstancedFoliageActorForCurrentLevel(GetWorld());
+	TInlineComponentArray<UActorComponent*> Components;
+	Actor->GetComponents(Components);
 
-		TInlineComponentArray<UActorComponent*> Components;
-		Actor->GetComponents(Components);
+	// Remove all instances of the selected meshes
+	for (const auto& MeshUIInfo : FoliageMeshList)
+	{
+		UFoliageType* FoliageType = MeshUIInfo->Settings;
+		if (!FoliageType->IsSelected)
+		{
+			continue;
+		}
 
-		for (auto Component : Components)
+		// Go through all FoliageActors in the world and delete 
+		for (FFoliageMeshInfoIterator It(World, FoliageType); It; ++It)
 		{
-			auto BaseId = IFA->InstanceBaseCache.GetInstanceBaseId(Component);
+			AInstancedFoliageActor* IFA = It.GetActor();
 			
-			for (auto& MeshPair : IFA->FoliageMeshes)
+			for (auto Component : Components)
 			{
-				FFoliageMeshInfo& MeshInfo = *MeshPair.Value;
-				UFoliageType* Settings = MeshPair.Key;
-
-				if (Settings->IsSelected)
-				{
-					auto* InstanceSet = MeshInfo.ComponentHash.Find(BaseId);
-					if (InstanceSet)
-					{
-						TArray<int32> InstancesToRemove = InstanceSet->Array();
-						MeshInfo.RemoveInstances(IFA, InstancesToRemove);
-					}
-				}
+				IFA->DeleteInstancesForComponent(Component, FoliageType);
 			}
 		}
+
+		OnInstanceCountUpdated(FoliageType);
 	}
-	else
-	{
-		TMap<UPrimitiveComponent*, TArray<FFoliagePaintBucketTriangle> > ComponentPotentialTriangles;
+}
+
+/** Apply paint bucket to actor */
+void FEdModeFoliage::ApplyPaintBucket_Add(AActor* Actor)
+{
+	UWorld* World = Actor->GetWorld();
+	TMap<UPrimitiveComponent*, TArray<FFoliagePaintBucketTriangle> > ComponentPotentialTriangles;
+
+	// Check all the components of the hit actor
+	TInlineComponentArray<UStaticMeshComponent*> StaticMeshComponents;
+	Actor->GetComponents(StaticMeshComponents);
 
-		// Check all the components of the hit actor
-		TInlineComponentArray<UStaticMeshComponent*> StaticMeshComponents;
-		Actor->GetComponents(StaticMeshComponents);
+	for (auto StaticMeshComponent : StaticMeshComponents)
+	{
+		UMaterialInterface* Material = StaticMeshComponent->GetMaterial(0);
 
-		for (int32 ComponentIdx = 0; ComponentIdx < StaticMeshComponents.Num(); ComponentIdx++)
+		if (UISettings.bFilterStaticMesh && StaticMeshComponent->StaticMesh && StaticMeshComponent->StaticMesh->RenderData &&
+			(UISettings.bFilterTranslucent || !Material || !IsTranslucentBlendMode(Material->GetBlendMode())))
 		{
-			UStaticMeshComponent* StaticMeshComponent = StaticMeshComponents[ComponentIdx];
-			UMaterialInterface* Material = StaticMeshComponent->GetMaterial(0);
+			UStaticMesh* StaticMesh = StaticMeshComponent->StaticMesh;
+			FStaticMeshLODResources& LODModel = StaticMesh->RenderData->LODResources[0];
+			TArray<FFoliagePaintBucketTriangle>& PotentialTriangles = ComponentPotentialTriangles.Add(StaticMeshComponent, TArray<FFoliagePaintBucketTriangle>());
 
-			if (UISettings.bFilterStaticMesh && StaticMeshComponent->StaticMesh && StaticMeshComponent->StaticMesh->RenderData &&
-				(UISettings.bFilterTranslucent || !Material || !IsTranslucentBlendMode(Material->GetBlendMode())))
+			bool bHasInstancedColorData = false;
+			const FStaticMeshComponentLODInfo* InstanceMeshLODInfo = nullptr;
+			if (StaticMeshComponent->LODData.Num() > 0)
 			{
-				UStaticMesh* StaticMesh = StaticMeshComponent->StaticMesh;
-				FStaticMeshLODResources& LODModel = StaticMesh->RenderData->LODResources[0];
-				TArray<FFoliagePaintBucketTriangle>& PotentialTriangles = ComponentPotentialTriangles.Add(StaticMeshComponent, TArray<FFoliagePaintBucketTriangle>());
+				InstanceMeshLODInfo = StaticMeshComponent->LODData.GetData();
+				bHasInstancedColorData = InstanceMeshLODInfo->PaintedVertices.Num() == LODModel.VertexBuffer.GetNumVertices();
+			}
 
-				bool bHasInstancedColorData = false;
-				const FStaticMeshComponentLODInfo* InstanceMeshLODInfo = nullptr;
-				if (StaticMeshComponent->LODData.Num() > 0)
-				{
-					InstanceMeshLODInfo = StaticMeshComponent->LODData.GetData();
-					bHasInstancedColorData = InstanceMeshLODInfo->PaintedVertices.Num() == LODModel.VertexBuffer.GetNumVertices();
-				}
+			const bool bHasColorData = bHasInstancedColorData || LODModel.ColorVertexBuffer.GetNumVertices();
 
-				const bool bHasColorData = bHasInstancedColorData || LODModel.ColorVertexBuffer.GetNumVertices();
+			// Get the raw triangle data for this static mesh
+			FTransform LocalToWorld = StaticMeshComponent->ComponentToWorld;
+			FIndexArrayView Indices = LODModel.IndexBuffer.GetArrayView();
+			const FPositionVertexBuffer& PositionVertexBuffer = LODModel.PositionVertexBuffer;
+			const FColorVertexBuffer& ColorVertexBuffer = LODModel.ColorVertexBuffer;
 
-				// Get the raw triangle data for this static mesh
-				FTransform LocalToWorld = StaticMeshComponent->ComponentToWorld;
-				FIndexArrayView Indices = LODModel.IndexBuffer.GetArrayView();
-				const FPositionVertexBuffer& PositionVertexBuffer = LODModel.PositionVertexBuffer;
-				const FColorVertexBuffer& ColorVertexBuffer = LODModel.ColorVertexBuffer;
+			if (USplineMeshComponent* SplineMesh = Cast<USplineMeshComponent>(StaticMeshComponent))
+			{
+				// Transform spline mesh verts correctly
+				FVector Mask = FVector(1, 1, 1);
+				USplineMeshComponent::GetAxisValue(Mask, SplineMesh->ForwardAxis) = 0;
 
-				if (USplineMeshComponent* SplineMesh = Cast<USplineMeshComponent>(StaticMeshComponent))
+				for (int32 Idx = 0; Idx < Indices.Num(); Idx += 3)
 				{
-					// Transform spline mesh verts correctly
-					FVector Mask = FVector(1, 1, 1);
-					USplineMeshComponent::GetAxisValue(Mask, SplineMesh->ForwardAxis) = 0;
-
-					for (int32 Idx = 0; Idx < Indices.Num(); Idx += 3)
-					{
-						const int32 Index0 = Indices[Idx];
-						const int32 Index1 = Indices[Idx + 1];
-						const int32 Index2 = Indices[Idx + 2];
-
-						const FVector Vert0 = SplineMesh->CalcSliceTransform(USplineMeshComponent::GetAxisValue(PositionVertexBuffer.VertexPosition(Index0), SplineMesh->ForwardAxis)).TransformPosition(PositionVertexBuffer.VertexPosition(Index0) * Mask);
-						const FVector Vert1 = SplineMesh->CalcSliceTransform(USplineMeshComponent::GetAxisValue(PositionVertexBuffer.VertexPosition(Index1), SplineMesh->ForwardAxis)).TransformPosition(PositionVertexBuffer.VertexPosition(Index1) * Mask);
-						const FVector Vert2 = SplineMesh->CalcSliceTransform(USplineMeshComponent::GetAxisValue(PositionVertexBuffer.VertexPosition(Index2), SplineMesh->ForwardAxis)).TransformPosition(PositionVertexBuffer.VertexPosition(Index2) * Mask);
-
-						new(PotentialTriangles)FFoliagePaintBucketTriangle(LocalToWorld
-							, Vert0
-							, Vert1
-							, Vert2
-							, bHasInstancedColorData ? InstanceMeshLODInfo->PaintedVertices[Index0].Color : (bHasColorData ? ColorVertexBuffer.VertexColor(Index0) : FColor::White)
-							, bHasInstancedColorData ? InstanceMeshLODInfo->PaintedVertices[Index1].Color : (bHasColorData ? ColorVertexBuffer.VertexColor(Index1) : FColor::White)
-							, bHasInstancedColorData ? InstanceMeshLODInfo->PaintedVertices[Index2].Color : (bHasColorData ? ColorVertexBuffer.VertexColor(Index2) : FColor::White)
-							);
-					}
+					const int32 Index0 = Indices[Idx];
+					const int32 Index1 = Indices[Idx + 1];
+					const int32 Index2 = Indices[Idx + 2];
+
+					const FVector Vert0 = SplineMesh->CalcSliceTransform(USplineMeshComponent::GetAxisValue(PositionVertexBuffer.VertexPosition(Index0), SplineMesh->ForwardAxis)).TransformPosition(PositionVertexBuffer.VertexPosition(Index0) * Mask);
+					const FVector Vert1 = SplineMesh->CalcSliceTransform(USplineMeshComponent::GetAxisValue(PositionVertexBuffer.VertexPosition(Index1), SplineMesh->ForwardAxis)).TransformPosition(PositionVertexBuffer.VertexPosition(Index1) * Mask);
+					const FVector Vert2 = SplineMesh->CalcSliceTransform(USplineMeshComponent::GetAxisValue(PositionVertexBuffer.VertexPosition(Index2), SplineMesh->ForwardAxis)).TransformPosition(PositionVertexBuffer.VertexPosition(Index2) * Mask);
+
+					new(PotentialTriangles)FFoliagePaintBucketTriangle(LocalToWorld
+						, Vert0
+						, Vert1
+						, Vert2
+						, bHasInstancedColorData ? InstanceMeshLODInfo->PaintedVertices[Index0].Color : (bHasColorData ? ColorVertexBuffer.VertexColor(Index0) : FColor::White)
+						, bHasInstancedColorData ? InstanceMeshLODInfo->PaintedVertices[Index1].Color : (bHasColorData ? ColorVertexBuffer.VertexColor(Index1) : FColor::White)
+						, bHasInstancedColorData ? InstanceMeshLODInfo->PaintedVertices[Index2].Color : (bHasColorData ? ColorVertexBuffer.VertexColor(Index2) : FColor::White)
+						);
 				}
-				else
+			}
+			else
+			{
+				// Build a mapping of vertex positions to vertex colors.  Using a TMap will allow for fast lookups so we can match new static mesh vertices with existing colors 
+				for (int32 Idx = 0; Idx < Indices.Num(); Idx += 3)
 				{
-					// Build a mapping of vertex positions to vertex colors.  Using a TMap will allow for fast lookups so we can match new static mesh vertices with existing colors 
-					for (int32 Idx = 0; Idx < Indices.Num(); Idx += 3)
-					{
-						const int32 Index0 = Indices[Idx];
-						const int32 Index1 = Indices[Idx + 1];
-						const int32 Index2 = Indices[Idx + 2];
-
-						new(PotentialTriangles)FFoliagePaintBucketTriangle(LocalToWorld
-							, PositionVertexBuffer.VertexPosition(Index0)
-							, PositionVertexBuffer.VertexPosition(Index1)
-							, PositionVertexBuffer.VertexPosition(Index2)
-							, bHasInstancedColorData ? InstanceMeshLODInfo->PaintedVertices[Index0].Color : (bHasColorData ? ColorVertexBuffer.VertexColor(Index0) : FColor::White)
-							, bHasInstancedColorData ? InstanceMeshLODInfo->PaintedVertices[Index1].Color : (bHasColorData ? ColorVertexBuffer.VertexColor(Index1) : FColor::White)
-							, bHasInstancedColorData ? InstanceMeshLODInfo->PaintedVertices[Index2].Color : (bHasColorData ? ColorVertexBuffer.VertexColor(Index2) : FColor::White)
-							);
-					}
+					const int32 Index0 = Indices[Idx];
+					const int32 Index1 = Indices[Idx + 1];
+					const int32 Index2 = Indices[Idx + 2];
+
+					new(PotentialTriangles)FFoliagePaintBucketTriangle(LocalToWorld
+						, PositionVertexBuffer.VertexPosition(Index0)
+						, PositionVertexBuffer.VertexPosition(Index1)
+						, PositionVertexBuffer.VertexPosition(Index2)
+						, bHasInstancedColorData ? InstanceMeshLODInfo->PaintedVertices[Index0].Color : (bHasColorData ? ColorVertexBuffer.VertexColor(Index0) : FColor::White)
+						, bHasInstancedColorData ? InstanceMeshLODInfo->PaintedVertices[Index1].Color : (bHasColorData ? ColorVertexBuffer.VertexColor(Index1) : FColor::White)
+						, bHasInstancedColorData ? InstanceMeshLODInfo->PaintedVertices[Index2].Color : (bHasColorData ? ColorVertexBuffer.VertexColor(Index2) : FColor::White)
+						);
 				}
 			}
 		}
+	}
 
-		// Place foliage
-		AInstancedFoliageActor* IFA = AInstancedFoliageActor::GetInstancedFoliageActorForCurrentLevel(GetWorld());
-
-		for (auto& MeshPair : IFA->FoliageMeshes)
+	for (const auto& MeshUIInfo : FoliageMeshList)
+	{
+		UFoliageType* Settings = MeshUIInfo->Settings;
+		if (!Settings->IsSelected)
 		{
-			FFoliageMeshInfo& MeshInfo = *MeshPair.Value;
-			UFoliageType* Settings = MeshPair.Key;
-
-			if (Settings->IsSelected)
-			{
-				// Quick lookup of potential instance locations, used for overlapping check.
-				TArray<FVector> PotentialInstanceLocations;
-				FFoliageInstanceHash PotentialInstanceHash(7);	// use 128x128 cell size, as the brush radius is typically small.
-				TArray<FPotentialInstance> InstancesToPlace;
-
-				// Radius where we expect to have a single instance, given the density rules
-				const float DensityCheckRadius = FMath::Max<float>(FMath::Sqrt((1000.f*1000.f) / (PI * Settings->Density * UISettings.GetPaintDensity())), Settings->Radius);
+			continue;
+		}
 
-				for (TMap<UPrimitiveComponent*, TArray<FFoliagePaintBucketTriangle> >::TIterator ComponentIt(ComponentPotentialTriangles); ComponentIt; ++ComponentIt)
-				{
-					UPrimitiveComponent* Component = ComponentIt.Key();
-					TArray<FFoliagePaintBucketTriangle>& PotentialTriangles = ComponentIt.Value();
+		// Quick lookup of potential instance locations, used for overlapping check.
+		TArray<FVector> PotentialInstanceLocations;
+		FFoliageInstanceHash PotentialInstanceHash(7);	// use 128x128 cell size, as the brush radius is typically small.
+		TArray<FPotentialInstance> InstancesToPlace;
 
-					for (int32 TriIdx = 0; TriIdx<PotentialTriangles.Num(); TriIdx++)
-					{
-						FFoliagePaintBucketTriangle& Triangle = PotentialTriangles[TriIdx];
+		// Radius where we expect to have a single instance, given the density rules
+		const float DensityCheckRadius = FMath::Max<float>(FMath::Sqrt((1000.f*1000.f) / (PI * Settings->Density * UISettings.GetPaintDensity())), Settings->Radius);
 
-						// Check if we can reject this triangle based on normal.
-						if ((Settings->GroundSlope > 0.f && Triangle.WorldNormal.Z <= FMath::Cos(PI * Settings->GroundSlope / 180.f) - SMALL_NUMBER) ||
-							(Settings->GroundSlope < 0.f && Triangle.WorldNormal.Z >= FMath::Cos(-PI * Settings->GroundSlope / 180.f) + SMALL_NUMBER))
-						{
-							continue;
-						}
-
-						// This is the total set of instances disregarding parameters like slope, height or layer.
-						float DesiredInstanceCountFloat = Triangle.Area * Settings->Density * UISettings.GetPaintDensity() / (1000.f*1000.f);
+		for (TMap<UPrimitiveComponent*, TArray<FFoliagePaintBucketTriangle> >::TIterator ComponentIt(ComponentPotentialTriangles); ComponentIt; ++ComponentIt)
+		{
+			UPrimitiveComponent* Component = ComponentIt.Key();
+			TArray<FFoliagePaintBucketTriangle>& PotentialTriangles = ComponentIt.Value();
 
-						// Allow a single instance with a random chance, if the brush is smaller than the density
-						int32 DesiredInstanceCount = DesiredInstanceCountFloat > 1.f ? FMath::RoundToInt(DesiredInstanceCountFloat) : FMath::FRand() < DesiredInstanceCountFloat ? 1 : 0;
+			for (int32 TriIdx = 0; TriIdx<PotentialTriangles.Num(); TriIdx++)
+			{
+				FFoliagePaintBucketTriangle& Triangle = PotentialTriangles[TriIdx];
 
-						for (int32 Idx = 0; Idx < DesiredInstanceCount; Idx++)
-						{
-							FVector InstLocation;
-							FColor VertexColor;
-							Triangle.GetRandomPoint(InstLocation, VertexColor);
+				// Check if we can reject this triangle based on normal.
+				if ((Settings->GroundSlope > 0.f && Triangle.WorldNormal.Z <= FMath::Cos(PI * Settings->GroundSlope / 180.f) - SMALL_NUMBER) ||
+					(Settings->GroundSlope < 0.f && Triangle.WorldNormal.Z >= FMath::Cos(-PI * Settings->GroundSlope / 180.f) + SMALL_NUMBER))
+				{
+					continue;
+				}
 
-							// Check color mask and filters at this location
-							if (!CheckVertexColor(Settings, VertexColor) ||
-								!CheckLocationForPotentialInstance(MeshInfo, Settings, DensityCheckRadius, InstLocation, Triangle.WorldNormal, PotentialInstanceLocations, PotentialInstanceHash))
-							{
-								continue;
-							}
+				// This is the total set of instances disregarding parameters like slope, height or layer.
+				float DesiredInstanceCountFloat = Triangle.Area * Settings->Density * UISettings.GetPaintDensity() / (1000.f*1000.f);
 
-							new(InstancesToPlace)FPotentialInstance(InstLocation, Triangle.WorldNormal, Component, 1.f);
-						}
-					}
-				}
+				// Allow a single instance with a random chance, if the brush is smaller than the density
+				int32 DesiredInstanceCount = DesiredInstanceCountFloat > 1.f ? FMath::RoundToInt(DesiredInstanceCountFloat) : FMath::FRand() < DesiredInstanceCountFloat ? 1 : 0;
 
-				// We need a world pointer
-				UWorld* World = Actor->GetWorld();
-				check(World);
-				// Place instances
-				for (int32 Idx = 0; Idx < InstancesToPlace.Num(); Idx++)
+				for (int32 Idx = 0; Idx < DesiredInstanceCount; Idx++)
 				{
-					FFoliageInstance Inst;
-					if (InstancesToPlace[Idx].PlaceInstance(Settings, Inst))
+					FVector InstLocation;
+					FColor VertexColor;
+					Triangle.GetRandomPoint(InstLocation, VertexColor);
+
+					// Check color mask and filters at this location
+					if (!CheckVertexColor(Settings, VertexColor) ||
+						!CheckLocationForPotentialInstance(World, Settings, DensityCheckRadius, InstLocation, Triangle.WorldNormal, PotentialInstanceLocations, PotentialInstanceHash))
 					{
-						MeshInfo.AddInstance(IFA, Settings, Inst);
+						continue;
 					}
+
+					new(InstancesToPlace)FPotentialInstance(InstLocation, Triangle.WorldNormal, Component, 1.f);
 				}
 			}
 		}
+		
+		// Place instances
+		for (FPotentialInstance& PotentialInstance : InstancesToPlace)
+		{
+			FFoliageInstance Inst;
+			if (PotentialInstance.PlaceInstance(World, Settings, Inst))
+			{
+				SpawnFoliageInstance(World, Settings, Inst, PotentialInstance.HitComponent);
+			}
+		}
+
+		//
+		OnInstanceCountUpdated(Settings);
 	}
 }
 
@@ -1618,6 +1942,52 @@ bool FEdModeFoliage::GetStaticMeshVertexColorForHit(const UStaticMeshComponent*
 	return false;
 }
 
+void FEdModeFoliage::SnapSelectedInstancesToGround(UWorld* InWorld)
+{
+	bool bMovedInstance = false;
+	
+	const int32 NumLevels = InWorld->GetNumLevels();
+	for (int32 LevelIdx = 0; LevelIdx < NumLevels; ++LevelIdx)
+	{
+		ULevel* Level = InWorld->GetLevel(LevelIdx);
+		AInstancedFoliageActor* IFA = AInstancedFoliageActor::GetInstancedFoliageActorForLevel(Level);
+		if (IFA)
+		{
+			bool bFoundSelection = false;
+
+			for (auto& MeshPair : IFA->FoliageMeshes)
+			{
+				FFoliageMeshInfo& MeshInfo = *MeshPair.Value;
+				TArray<int32> SelectedIndices = MeshInfo.SelectedIndices.Array();
+
+				if (SelectedIndices.Num() > 0)
+				{
+					// Mark actor once we found selection
+					if (!bFoundSelection)
+					{
+						IFA->Modify();
+						bFoundSelection = true;
+					}
+
+					MeshInfo.PreMoveInstances(IFA, SelectedIndices);
+
+					for (int32 InstanceIndex : SelectedIndices)
+					{
+						bMovedInstance|= SnapInstanceToGround(IFA, MeshPair.Key->AlignMaxAngle, MeshInfo, InstanceIndex);
+					}
+
+					MeshInfo.PostMoveInstances(IFA, SelectedIndices);
+				}
+			}
+		}
+	}
+
+	if (bMovedInstance)
+	{
+		UpdateWidgetLocationToInstanceSelection();
+	}
+}
+
 bool FEdModeFoliage::SnapInstanceToGround(AInstancedFoliageActor* InIFA, float AlignMaxAngle, FFoliageMeshInfo& Mesh, int32 InstanceIdx)
 {
 	FFoliageInstance& Instance = Mesh.Instances[InstanceIdx];
@@ -1626,7 +1996,7 @@ bool FEdModeFoliage::SnapInstanceToGround(AInstancedFoliageActor* InIFA, float A
 
 	FHitResult Hit;
 	static FName NAME_FoliageSnap = FName("FoliageSnap");
-	if (AInstancedFoliageActor::FoliageTrace(InIFA->GetWorld(), Hit, FDesiredFoliageInstance(Start, End), InIFA, NAME_FoliageSnap))
+	if (AInstancedFoliageActor::FoliageTrace(InIFA->GetWorld(), Hit, FDesiredFoliageInstance(Start, End), NAME_FoliageSnap))
 	{
 		UPrimitiveComponent* HitComponent = Hit.Component.Get();
 
@@ -1668,97 +2038,179 @@ bool FEdModeFoliage::SnapInstanceToGround(AInstancedFoliageActor* InIFA, float A
 	return false;
 }
 
-
-/** Get list of meshs for current level for UI */
-
-TArray<struct FFoliageMeshUIInfo>& FEdModeFoliage::GetFoliageMeshList()
+TArray<FFoliageMeshUIInfoPtr>& FEdModeFoliage::GetFoliageMeshList()
 {
-	UpdateFoliageMeshList();
 	return FoliageMeshList;
 }
 
-void FEdModeFoliage::UpdateFoliageMeshList()
+void FEdModeFoliage::PopulateFoliageMeshList()
 {
 	FoliageMeshList.Empty();
+	
+	// Collect set of all available Settings
+	UWorld* World = GetWorld();
+	ULevel* CurrentLevel = World->GetCurrentLevel();
+	const int32 NumLevels = World->GetNumLevels();
+	
+	for (int32 LevelIdx = 0; LevelIdx < NumLevels; ++LevelIdx)
+	{
+		ULevel* Level = World->GetLevel(LevelIdx);
+		if (Level && Level->bIsVisible)
+		{
+			AInstancedFoliageActor* IFA = AInstancedFoliageActor::GetInstancedFoliageActorForLevel(Level);
+			if (IFA)
+			{
+				for (auto& MeshPair : IFA->FoliageMeshes)
+				{
+					if (!CanPaint(MeshPair.Key, CurrentLevel))
+					{
+						continue;
+					}
+										
+					int32 ElementIdx = FoliageMeshList.IndexOfByPredicate([&](const FFoliageMeshUIInfoPtr& Item) {
+						return Item->Settings == MeshPair.Key;
+					});
 
-	AInstancedFoliageActor* IFA = AInstancedFoliageActor::GetInstancedFoliageActorForCurrentLevel(GetWorld());
+					if (ElementIdx == INDEX_NONE)
+					{
+						ElementIdx = FoliageMeshList.Add(MakeShareable(new FFoliageMeshUIInfo(MeshPair.Key)));
+					}
 
-	for (auto& MeshPair : IFA->FoliageMeshes)
-	{
-		new(FoliageMeshList)FFoliageMeshUIInfo(MeshPair.Key, &*MeshPair.Value);
+					FoliageMeshList[ElementIdx]->InstanceCountTotal+= MeshPair.Value->GetInstanceCount();
+					if (Level == World->GetCurrentLevel())
+					{
+						FoliageMeshList[ElementIdx]->InstanceCountCurrentLevel+= MeshPair.Value->GetInstanceCount();
+					}
+				}
+			}
+		}
 	}
 
-	auto CompareDisplayOrder = [](const FFoliageMeshUIInfo& A, const FFoliageMeshUIInfo& B)
+	auto CompareDisplayOrder = [](const FFoliageMeshUIInfoPtr& A, const FFoliageMeshUIInfoPtr& B)
 	{
-		return A.Settings->DisplayOrder < B.Settings->DisplayOrder;
+		return A->Settings->DisplayOrder < B->Settings->DisplayOrder;
 	};
+
 	FoliageMeshList.Sort(CompareDisplayOrder);
+
+	StaticCastSharedPtr<FFoliageEdModeToolkit>(Toolkit)->RefreshFullList();
 }
 
-/** Add a new mesh (if it doesn't already exist) */
-void FEdModeFoliage::AddFoliageMesh(UStaticMesh* Mesh)
+void FEdModeFoliage::OnInstanceCountUpdated(const UFoliageType* FoliageType)
 {
-	AInstancedFoliageActor* IFA = AInstancedFoliageActor::GetInstancedFoliageActorForCurrentLevel(GetWorld());
-
-	if (IFA->GetSettingsForMesh(Mesh) == nullptr)
+	int32 EntryIndex = FoliageMeshList.IndexOfByPredicate([&](const FFoliageMeshUIInfoPtr& UIInfoPtr) {
+		return UIInfoPtr->Settings == FoliageType;
+	});
+	
+	if (EntryIndex == INDEX_NONE)
 	{
-		IFA->AddMesh(Mesh);
+		return;
+	}
+	
+	int32 InstanceCountTotal = 0;
+	int32 InstanceCountCurrentLevel = 0;
+	UWorld* World = GetWorld();
+	ULevel* CurrentLevel = World->GetCurrentLevel();
 
-		// Update mesh list.
-		UpdateFoliageMeshList();
+	for (FFoliageMeshInfoIterator It(World, FoliageType); It; ++It)
+	{
+		FFoliageMeshInfo* MeshInfo = (*It);
+		InstanceCountTotal+= MeshInfo->Instances.Num();
+		if (It.GetActor()->GetLevel() == CurrentLevel)
+		{
+			InstanceCountCurrentLevel = MeshInfo->Instances.Num();
+		}
 	}
+	
+	//
+	FoliageMeshList[EntryIndex]->InstanceCountTotal = InstanceCountTotal;
+	FoliageMeshList[EntryIndex]->InstanceCountCurrentLevel = InstanceCountCurrentLevel;
 }
 
-void FEdModeFoliage::AddFoliageMesh(UFoliageType* Settings)
+bool FEdModeFoliage::CanPaint(const UFoliageType* FoliageType, const ULevel* InLevel)
 {
-	AInstancedFoliageActor* IFA = AInstancedFoliageActor::GetInstancedFoliageActorForCurrentLevel(GetWorld());
-	if (IFA->FindMesh(Settings) == nullptr)
+	// Non-shared objects can be painted only into their own level
+	// Assets can be painted everywhere
+	if (FoliageType->IsAsset() || FoliageType->GetOutermost() == InLevel->GetOutermost())
 	{
-		IFA->AddMesh(Settings);
+		return true;
+	}
+	
+	return false;
+}
 
-		// Update mesh list.
-		UpdateFoliageMeshList();
+UFoliageType* FEdModeFoliage::AddFoliageAsset(UObject* InAssset)
+{
+	UFoliageType* FoliageType = nullptr;
+	
+	UStaticMesh* StaticMesh = Cast<UStaticMesh>(InAssset);
+	if (StaticMesh)
+	{
+		AInstancedFoliageActor* IFA = AInstancedFoliageActor::GetInstancedFoliageActorForCurrentLevel(GetWorld(), true);
+		FoliageType = IFA->GetSettingsForMesh(StaticMesh);
+		if (!FoliageType)
+		{
+			IFA->AddMesh(StaticMesh, &FoliageType);
+		}
+	}
+	else 
+	{
+		FoliageType = Cast<UFoliageType>(InAssset);
+		if (FoliageType)
+		{
+			AInstancedFoliageActor* IFA = AInstancedFoliageActor::GetInstancedFoliageActorForCurrentLevel(GetWorld(), true);
+			FoliageType = IFA->AddFoliageType(FoliageType);
+		}
+	}
+	
+	if (FoliageType)
+	{
+		PopulateFoliageMeshList();
 	}
+	
+	return FoliageType;
 }
 
 /** Remove a mesh */
-bool FEdModeFoliage::RemoveFoliageMesh(UFoliageType* Settings)
+bool FEdModeFoliage::RemoveFoliageType(UFoliageType** FoliageTypeList, int32 Num)
 {
-	bool bMeshRemoved = false;
-
-	AInstancedFoliageActor* IFA = AInstancedFoliageActor::GetInstancedFoliageActorForCurrentLevel(GetWorld());
-	FFoliageMeshInfo* MeshInfo = IFA->FindMesh(Settings);
-	if (MeshInfo != nullptr)
+	TArray<AInstancedFoliageActor*> IFAList;
+	// Find all foliage actors that have any of these types
+	UWorld* World = GetWorld();
+	for (int32 FoliageTypeIdx = 0; FoliageTypeIdx < Num; ++FoliageTypeIdx)
 	{
-		int32 InstancesNum = MeshInfo->Instances.Num();
-
-		bool bProceed = true;
-		if (Settings->GetStaticMesh() != nullptr && InstancesNum > 0)
+		UFoliageType* FoliageType = FoliageTypeList[FoliageTypeIdx];
+		for (FFoliageMeshInfoIterator It(World, FoliageType); It; ++It)
 		{
-			FText Message = FText::Format(NSLOCTEXT("UnrealEd", "FoliageMode_DeleteMesh", "Are you sure you want to remove all {0} instances of {1} from this level?"), FText::AsNumber(InstancesNum), FText::FromName(Settings->GetStaticMesh()->GetFName()));
-			bProceed = (FMessageDialog::Open(EAppMsgType::YesNo, Message) == EAppReturnType::Yes);
+			IFAList.Add(It.GetActor());
 		}
+	}
 
-		if (bProceed)
+	if (IFAList.Num())
+	{
+		GEditor->BeginTransaction(NSLOCTEXT("UnrealEd", "FoliageMode_RemoveMeshTransaction", "Foliage Editing: Remove Mesh"));
+		for (AInstancedFoliageActor* IFA : IFAList)
 		{
-			GEditor->BeginTransaction(NSLOCTEXT("UnrealEd", "FoliageMode_RemoveMeshTransaction", "Foliage Editing: Remove Mesh"));
-			IFA->RemoveMesh(Settings);
-			GEditor->EndTransaction();
-
-			bMeshRemoved = true;
+			IFA->RemoveFoliageType(FoliageTypeList, Num);
 		}
+		GEditor->EndTransaction();
 
-		// Update mesh list.
-		UpdateFoliageMeshList();
+		PopulateFoliageMeshList();
+		return true;
 	}
 
-	return bMeshRemoved;
+	return false;
 }
 
 /** Bake instances to StaticMeshActors */
 void FEdModeFoliage::BakeFoliage(UFoliageType* Settings, bool bSelectedOnly)
 {
 	AInstancedFoliageActor* IFA = AInstancedFoliageActor::GetInstancedFoliageActorForCurrentLevel(GetWorld());
+	if (IFA == nullptr)
+	{
+		return;
+	}
+
 	FFoliageMeshInfo* MeshInfo = IFA->FindMesh(Settings);
 	if (MeshInfo != nullptr)
 	{
@@ -1821,127 +2273,102 @@ void FEdModeFoliage::ReplaceSettingsObject(UFoliageType* OldSettings, UFoliageTy
 {
 	FScopedTransaction Transaction(NSLOCTEXT("UnrealEd", "FoliageMode_ReplaceSettingsObject", "Foliage Editing: Replace Settings Object"));
 
-	AInstancedFoliageActor* IFA = AInstancedFoliageActor::GetInstancedFoliageActorForCurrentLevel(GetWorld());
-	IFA->Modify();
-
-	TUniqueObj<FFoliageMeshInfo> MeshInfo;
-	if (IFA->FoliageMeshes.RemoveAndCopyValue(OldSettings, MeshInfo))
+	UWorld* World = GetWorld();
+	for (FFoliageMeshInfoIterator It(World, OldSettings); It; ++It)
 	{
-		IFA->FoliageMeshes.Add(NewSettings, MoveTemp(MeshInfo));
-	}
-	else
-	{
-		Transaction.Cancel();
+		AInstancedFoliageActor* IFA = It.GetActor();
+			
+		IFA->Modify();
+		TUniqueObj<FFoliageMeshInfo> OldMeshInfo;
+		IFA->FoliageMeshes.RemoveAndCopyValue(OldSettings, OldMeshInfo);
+
+		// Append instances if new foliage type is already exists in this actor
+		// Otherwise just replace key entry for instances
+		TUniqueObj<FFoliageMeshInfo>* NewMeshInfo = IFA->FoliageMeshes.Find(NewSettings);
+		if (NewMeshInfo)
+		{
+			(*NewMeshInfo)->Instances.Append(OldMeshInfo->Instances);
+			// Old component needs to go
+			if (OldMeshInfo->Component != nullptr)
+			{
+				OldMeshInfo->Component->bAutoRegister = false;
+				OldMeshInfo->Component = nullptr;
+			}
+			(*NewMeshInfo)->ReallocateClusters(IFA, NewSettings);
+		}
+		else
+		{
+			IFA->FoliageMeshes.Add(NewSettings, MoveTemp(OldMeshInfo))->ReallocateClusters(IFA, NewSettings);
+		}
 	}
+			
+	PopulateFoliageMeshList();
 }
 
-/** Save the settings object */
-UFoliageType* FEdModeFoliage::SaveSettingsObject(const FText& InSettingsPackageName, UFoliageType* Settings)
+UFoliageType* FEdModeFoliage::SaveSettingsObject(UFoliageType* Settings)
 {
 	FScopedTransaction Transaction(NSLOCTEXT("UnrealEd", "FoliageMode_SaveSettingsObject", "Foliage Editing: Save Settings Object"));
 
-	AInstancedFoliageActor* IFA = AInstancedFoliageActor::GetInstancedFoliageActorForCurrentLevel(GetWorld());
-	IFA->Modify();
-
-	FString PackageName = InSettingsPackageName.ToString();
-	UPackage* Package = CreatePackage(nullptr, *PackageName);
-	check(Package);
+	UFoliageType* SettingsToSave = nullptr;
 
-	UFoliageType* NewSettings = Cast<UFoliageType>(StaticDuplicateObject(Settings, Package, *FPackageName::GetLongPackageAssetName(PackageName)));
-	NewSettings->SetFlags(RF_Standalone | RF_Public);
-	NewSettings->Modify();
-
-	ReplaceSettingsObject(Settings, NewSettings);
+	if (!Settings->IsAsset())
+	{
+		FString PackageName;
+		UStaticMesh* StaticMesh = Settings->GetStaticMesh();
+		if (StaticMesh)
+		{
+			// Build default settings asset name and path
+			PackageName = FPackageName::GetLongPackagePath(StaticMesh->GetOutermost()->GetName()) + TEXT("/") + StaticMesh->GetName() + TEXT("_settings");
+		}
+		
+		TSharedRef<SDlgPickAssetPath> SettingDlg =
+			SNew(SDlgPickAssetPath)
+			.Title(LOCTEXT("SettingsDialogTitle", "Choose Location for Foliage Settings Asset"))
+			.DefaultAssetPath(FText::FromString(PackageName));
 
-	// Notify the asset registry
-	FAssetRegistryModule::AssetCreated(NewSettings);
+		if (SettingDlg->ShowModal() != EAppReturnType::Cancel)
+		{
+			PackageName = SettingDlg->GetFullAssetPath().ToString();
+			UPackage* Package = CreatePackage(nullptr, *PackageName);
+			SettingsToSave = Cast<UFoliageType>(StaticDuplicateObject(Settings, Package, *FPackageName::GetLongPackageAssetName(PackageName)));
+			SettingsToSave->SetFlags(RF_Standalone | RF_Public);
+			SettingsToSave->Modify();
 
-	return NewSettings;
-}
+			// Notify the asset registry
+			FAssetRegistryModule::AssetCreated(SettingsToSave);
 
-/** Reapply cluster settings to all the instances */
-void FEdModeFoliage::ReallocateClusters(UFoliageType* Settings)
-{
-	AInstancedFoliageActor* IFA = AInstancedFoliageActor::GetInstancedFoliageActorForCurrentLevel(GetWorld());
-	FFoliageMeshInfo* MeshInfo = IFA->FindMesh(Settings);
-	if (MeshInfo != nullptr)
+			ReplaceSettingsObject(Settings, SettingsToSave);
+		}
+	}
+	else
 	{
-		MeshInfo->ReallocateClusters(IFA, Settings);
+		SettingsToSave = Settings;
 	}
+	
+	// Save to disk
+	if (SettingsToSave)
+	{
+		TArray<UPackage*> PackagesToSave; 
+		PackagesToSave.Add(SettingsToSave->GetOutermost());
+		const bool bCheckDirty = false;
+		const bool bPromptToSave = false;
+		FEditorFileUtils::PromptForCheckoutAndSave(PackagesToSave, bCheckDirty, bPromptToSave);
+	}
+		
+	return SettingsToSave;
 }
 
-/** Replace a mesh with another one */
-bool FEdModeFoliage::ReplaceStaticMesh(UFoliageType* OldSettings, UStaticMesh* NewStaticMesh, bool& bOutMeshMerged)
+/** Reapply cluster settings to all the instances */
+void FEdModeFoliage::ReallocateClusters(UFoliageType* Settings)
 {
-	bOutMeshMerged = false;
-
-	AInstancedFoliageActor* IFA = AInstancedFoliageActor::GetInstancedFoliageActorForCurrentLevel(GetWorld());
-	FFoliageMeshInfo* OldMeshInfo = IFA->FindMesh(OldSettings);
-
-	if (OldMeshInfo != nullptr && OldSettings->GetStaticMesh() != NewStaticMesh)
+	UWorld* World = GetWorld();
+	for (FFoliageMeshInfoIterator It(World, Settings); It; ++It)
 	{
-		int32 InstancesNum = OldMeshInfo->Instances.Num();
-
-		// Look for the new mesh in the mesh list, and either create a new mesh or merge the instances.
-		UFoliageType* NewSettings = nullptr;
-		FFoliageMeshInfo* NewMeshInfo = nullptr;
-		NewSettings = IFA->GetSettingsForMesh(NewStaticMesh, &NewMeshInfo);
-
-		if (NewMeshInfo == nullptr)
-		{
-			FText Message = FText::Format(NSLOCTEXT("UnrealEd", "FoliageMode_ReplaceMesh", "Are you sure you want to replace all {0} instances of {1} with {2}?"), FText::AsNumber(InstancesNum), FText::FromString(OldSettings->GetStaticMesh()->GetName()), FText::FromString(NewStaticMesh->GetName()));
-
-			if (InstancesNum > 0 && EAppReturnType::Yes != FMessageDialog::Open(EAppMsgType::YesNo, Message))
-			{
-				return false;
-			}
-
-			GEditor->BeginTransaction(NSLOCTEXT("UnrealEd", "FoliageMode_ChangeStaticMeshTransaction", "Foliage Editing: Change StaticMesh"));
-			IFA->Modify();
-			NewMeshInfo = IFA->AddMesh(NewStaticMesh, &NewSettings, Cast<UFoliageType_InstancedStaticMesh>(OldSettings));
-			NewSettings->DisplayOrder = OldSettings->DisplayOrder;
-			NewSettings->ShowNothing = OldSettings->ShowNothing;
-			NewSettings->ShowPaintSettings = OldSettings->ShowPaintSettings;
-			NewSettings->ShowInstanceSettings = OldSettings->ShowInstanceSettings;
-		}
-		else if (InstancesNum > 0 &&
-			EAppReturnType::Yes != FMessageDialog::Open(EAppMsgType::YesNo,
-			FText::Format(NSLOCTEXT("UnrealEd", "FoliageMode_ReplaceMeshMerge", "Are you sure you want to merge all {0} instances of {1} into {2}?"),
-			FText::AsNumber(InstancesNum), FText::FromString(OldSettings->GetStaticMesh()->GetName()), FText::FromString(NewStaticMesh->GetName()))))
-		{
-			return false;
-		}
-		else
-		{
-			bOutMeshMerged = true;
-
-			GEditor->BeginTransaction(NSLOCTEXT("UnrealEd", "FoliageMode_ChangeStaticMeshTransaction", "Foliage Editing: Change StaticMesh"));
-			IFA->Modify();
-		}
-
-		if (InstancesNum > 0)
-		{
-			// copy instances from old to new.
-			for (FFoliageInstance& Instance : OldMeshInfo->Instances)
-			{
-					NewMeshInfo->AddInstance(IFA, NewSettings, Instance);
-				}
-			}
-
-		// Remove the old mesh.
-		IFA->RemoveMesh(OldSettings);
-
-		GEditor->EndTransaction();
-
-		// Update mesh list.
-		UpdateFoliageMeshList();
-
-		StaticCastSharedPtr<FFoliageEdModeToolkit>(Toolkit)->RefreshFullList();
+		FFoliageMeshInfo* MeshInfo = (*It);
+		MeshInfo->ReallocateClusters(It.GetActor(), Settings);
 	}
-	return true;
 }
 
-
 /** FEdMode: Called when a key is pressed */
 bool FEdModeFoliage::InputKey(FEditorViewportClient* ViewportClient, FViewport* Viewport, FKey Key, EInputEvent Event)
 {
@@ -1965,29 +2392,7 @@ bool FEdModeFoliage::InputKey(FEditorViewportClient* ViewportClient, FViewport*
 				if (!bToolActive)
 				{
 					GEditor->BeginTransaction(NSLOCTEXT("UnrealEd", "FoliageMode_EditTransaction", "Foliage Editing"));
-					InstanceSnapshot.Empty();
-
-					// Special setup beginning a stroke with the Reapply tool
-					// Necessary so we don't keep reapplying settings over and over for the same instances.
-					if (UISettings.GetReapplyToolSelected())
-					{
-						AInstancedFoliageActor* IFA = AInstancedFoliageActor::GetInstancedFoliageActorForCurrentLevel(GetWorld());
-						for (auto& MeshPair : IFA->FoliageMeshes)
-						{
-							FFoliageMeshInfo& MeshInfo = *MeshPair.Value;
-							if (MeshPair.Key->IsSelected)
-							{
-								// Take a snapshot of all the locations
-								InstanceSnapshot.Add(MeshPair.Key, FMeshInfoSnapshot(&MeshInfo));
-
-								// Clear the "FOLIAGE_Readjusted" flag
-								for (int32 Idx = 0; Idx < MeshInfo.Instances.Num(); Idx++)
-								{
-									MeshInfo.Instances[Idx].Flags &= (~FOLIAGE_Readjusted);
-								}
-							}
-						}
-					}
+					PreApplyBrush();
 					ApplyBrush(ViewportClient);
 					bToolActive = true;
 					return true;
@@ -2014,48 +2419,12 @@ bool FEdModeFoliage::InputKey(FEditorViewportClient* ViewportClient, FViewport*
 		{
 			if (Key == EKeys::Platform_Delete)
 			{
-				GEditor->BeginTransaction(NSLOCTEXT("UnrealEd", "FoliageMode_EditTransaction", "Foliage Editing"));
-				AInstancedFoliageActor* IFA = AInstancedFoliageActor::GetInstancedFoliageActorForCurrentLevel(GetWorld());
-				IFA->Modify();
-				for (auto& MeshPair : IFA->FoliageMeshes)
-				{
-					FFoliageMeshInfo& Mesh = *MeshPair.Value;
-					if (Mesh.SelectedIndices.Num() > 0)
-					{
-						TArray<int32> InstancesToDelete = Mesh.SelectedIndices.Array();
-						Mesh.RemoveInstances(IFA, InstancesToDelete);
-					}
-				}
-				GEditor->EndTransaction();
-
+				RemoveSelectedInstances(GetWorld());
 				return true;
 			}
 			else if (Key == EKeys::End)
 			{
-				// Snap instances to ground
-				AInstancedFoliageActor* IFA = AInstancedFoliageActor::GetInstancedFoliageActorForCurrentLevel(GetWorld());
-				IFA->Modify();
-				bool bMovedInstance = false;
-				for (auto& MeshPair : IFA->FoliageMeshes)
-				{
-					FFoliageMeshInfo& Mesh = *MeshPair.Value;
-
-					TArray<int32> SelectedIndices = Mesh.SelectedIndices.Array();
-
-					Mesh.PreMoveInstances(IFA, SelectedIndices);
-					for (int32 InstanceIndex : SelectedIndices)
-					{
-						bMovedInstance|= SnapInstanceToGround(IFA, MeshPair.Key->AlignMaxAngle, Mesh, InstanceIndex);
-								}
-					Mesh.PostMoveInstances(IFA, SelectedIndices);
-				}
-
-				if (bMovedInstance)
-				{
-					Owner->PivotLocation = Owner->SnappedLocation = IFA->GetSelectionLocation();
-					IFA->MarkComponentsRenderStateDirty();
-				}
-
+				SnapSelectedInstancesToGround(GetWorld());
 				return true;
 			}
 		}
@@ -2127,25 +2496,25 @@ void FEdModeFoliage::ForceRealTimeViewports(const bool bEnable, const bool bStor
 
 bool FEdModeFoliage::HandleClick(FEditorViewportClient* InViewportClient, HHitProxy *HitProxy, const FViewportClick &Click)
 {
-	AInstancedFoliageActor* IFA = AInstancedFoliageActor::GetInstancedFoliageActorForCurrentLevel(GetWorld());
-
 	if (UISettings.GetSelectToolSelected())
 	{
 		if (HitProxy && HitProxy->IsA(HInstancedStaticMeshInstance::StaticGetType()))
 		{
 			HInstancedStaticMeshInstance* SMIProxy = ((HInstancedStaticMeshInstance*)HitProxy);
-
-			IFA->SelectInstance(SMIProxy->Component, SMIProxy->InstanceIndex, Click.IsControlDown());
-
-			// Update pivot
-			Owner->PivotLocation = Owner->SnappedLocation = IFA->GetSelectionLocation();
+			AInstancedFoliageActor* IFA = AInstancedFoliageActor::GetInstancedFoliageActorForLevel(SMIProxy->Component->GetComponentLevel());
+			if (IFA)
+			{
+				IFA->SelectInstance(SMIProxy->Component, SMIProxy->InstanceIndex, Click.IsControlDown());
+				// Update pivot
+				UpdateWidgetLocationToInstanceSelection();
+			}
 		}
 		else
 		{
 			if (!Click.IsControlDown())
 			{
 				// Select none if not trying to toggle
-				IFA->SelectInstance(nullptr, -1, false);
+				SelectInstances(GetWorld(), false);
 			}
 		}
 
@@ -2156,7 +2525,16 @@ bool FEdModeFoliage::HandleClick(FEditorViewportClient* InViewportClient, HHitPr
 		if (HitProxy && HitProxy->IsA(HActor::StaticGetType()))
 		{
 			GEditor->BeginTransaction(NSLOCTEXT("UnrealEd", "FoliageMode_EditTransaction", "Foliage Editing"));
-			ApplyPaintBucket(((HActor*)HitProxy)->Actor, Click.IsShiftDown());
+			
+			if (Click.IsShiftDown())
+			{
+				ApplyPaintBucket_Remove(((HActor*)HitProxy)->Actor);
+			}
+			else
+			{
+				ApplyPaintBucket_Add(((HActor*)HitProxy)->Actor);
+			}
+						
 			GEditor->EndTransaction();
 		}
 
@@ -2178,8 +2556,7 @@ bool FEdModeFoliage::StartTracking(FEditorViewportClient* InViewportClient, FVie
 	if (UISettings.GetSelectToolSelected() || UISettings.GetLassoSelectToolSelected())
 	{
 		// Update pivot
-		AInstancedFoliageActor* IFA = AInstancedFoliageActor::GetInstancedFoliageActorForCurrentLevel(GetWorld());
-		Owner->PivotLocation = Owner->SnappedLocation = IFA->GetSelectionLocation();
+		UpdateWidgetLocationToInstanceSelection();
 
 		GEditor->BeginTransaction(NSLOCTEXT("UnrealEd", "FoliageMode_EditTransaction", "Foliage Editing"));
 
@@ -2210,42 +2587,12 @@ bool FEdModeFoliage::InputDelta(FEditorViewportClient* InViewportClient, FViewpo
 
 	if (InViewportClient->GetCurrentWidgetAxis() != EAxisList::None && (UISettings.GetSelectToolSelected() || UISettings.GetLassoSelectToolSelected()))
 	{
-		AInstancedFoliageActor* IFA = AInstancedFoliageActor::GetInstancedFoliageActorForCurrentLevel(GetWorld());
-		IFA->Modify();
-		for (auto& MeshPair : IFA->FoliageMeshes)
-		{
-			FFoliageMeshInfo& MeshInfo = *MeshPair.Value;
-			TArray<int32> SelectedIndices = MeshInfo.SelectedIndices.Array();
-
-			bFoundSelection |= SelectedIndices.Num() > 0;
-
-			if (bAltDown && bCanAltDrag && (InViewportClient->GetCurrentWidgetAxis() & EAxisList::XYZ))
-			{
-				MeshInfo.DuplicateInstances(IFA, MeshPair.Key, SelectedIndices);
-			}
-
-			MeshInfo.PreMoveInstances(IFA, SelectedIndices);
-
-			for (int32 SelectedInstanceIdx : SelectedIndices)
-			{
-				FFoliageInstance& Instance = MeshInfo.Instances[SelectedInstanceIdx];
-				Instance.Location += InDrag;
-				Instance.ZOffset = 0.f;
-				Instance.Rotation += InRot;
-				Instance.DrawScale3D += InScale;
-			}
-
-			MeshInfo.PostMoveInstances(IFA, SelectedIndices);
-		}
+		const bool bDuplicateInstances = (bAltDown && bCanAltDrag && (InViewportClient->GetCurrentWidgetAxis() & EAxisList::XYZ));
+		
+		TransformSelectedInstances(GetWorld(), InDrag, InRot, InScale, bDuplicateInstances);
 
 		// Only allow alt-drag on first InputDelta
 		bCanAltDrag = false;
-
-		if (bFoundSelection)
-		{
-			IFA->MarkComponentsRenderStateDirty();
-			return true;
-		}
 	}
 
 	return FEdMode::InputDelta(InViewportClient, InViewport, InDrag, InRot, InScale);
@@ -2263,7 +2610,12 @@ bool FEdModeFoliage::UsesTransformWidget() const
 
 bool FEdModeFoliage::ShouldDrawWidget() const
 {
-	return (UISettings.GetSelectToolSelected() || (UISettings.GetLassoSelectToolSelected() && !bToolActive)) && AInstancedFoliageActor::GetInstancedFoliageActorForCurrentLevel(GetWorld())->SelectedMesh != nullptr;
+	if (UISettings.GetSelectToolSelected() || (UISettings.GetLassoSelectToolSelected() && !bToolActive))
+	{
+		FVector Location;
+		return GetSelectionLocation(GetWorld(), Location) != nullptr;
+	}
+	return false;
 }
 
 EAxisList::Type FEdModeFoliage::GetWidgetAxisToDraw(FWidget::EWidgetMode InWidgetMode) const
diff --git a/Engine/Source/Editor/FoliageEdit/Private/FoliageEdMode.h b/Engine/Source/Editor/FoliageEdit/Private/FoliageEdMode.h
--- a/Engine/Source/Editor/FoliageEdit/Private/FoliageEdMode.h
+++ b/Engine/Source/Editor/FoliageEdit/Private/FoliageEdMode.h
@@ -104,15 +104,24 @@ public:
 
 struct FFoliageMeshUIInfo
 {
-	UFoliageType* Settings;
-	FFoliageMeshInfo* MeshInfo;
-
-	FFoliageMeshUIInfo(UFoliageType* InSettings, FFoliageMeshInfo* InMeshInfo)
+	UFoliageType*	Settings;
+	int32			InstanceCountCurrentLevel;
+	int32			InstanceCountTotal;
+	
+	FFoliageMeshUIInfo(UFoliageType* InSettings)
 		: Settings(InSettings)
-		, MeshInfo(InMeshInfo)
+		, InstanceCountCurrentLevel(0)
+		, InstanceCountTotal(0)
 	{}
+
+	bool operator == (const FFoliageMeshUIInfo& Other) const
+	{
+		return Settings == Other.Settings;
+	}
 };
 
+typedef TSharedPtr<FFoliageMeshUIInfo> FFoliageMeshUIInfoPtr;
+
 // Snapshot of current MeshInfo state. Created at start of a brush stroke to store the existing instance info.
 class FMeshInfoSnapshot
 {
@@ -177,7 +186,9 @@ public:
 
 	/** Called when the current level changes */
 	void NotifyNewCurrentLevel();
-
+	void NotifyLevelAddedToWorld(ULevel* InLevel, UWorld* InWorld);
+	void NotifyLevelRemovedFromWorld(ULevel* InLevel, UWorld* InWorld);
+	
 	/** Called when the user changes the current tool in the UI */
 	void NotifyToolChanged();
 
@@ -257,27 +268,33 @@ public:
 	    returns a line segment inside the sphere parallel to the view direction */
 	void GetRandomVectorInBrush(FVector& OutStart, FVector& OutEnd);
 
+	/** Setup before call to ApplyBrush */
+	void PreApplyBrush();
+
 	/** Apply brush */
 	void ApplyBrush(FEditorViewportClient* ViewportClient);
 
-	/** Update existing mesh info for current level */
-	void UpdateFoliageMeshList();
+	/** Get list of meshes for current level */
+	TArray<FFoliageMeshUIInfoPtr>& GetFoliageMeshList();
 
-	TArray<struct FFoliageMeshUIInfo>& GetFoliageMeshList();
+	/** Populate mesh with foliage mesh settings found across world */
+	void PopulateFoliageMeshList();
 
-	/** Add a new mesh */
-	void AddFoliageMesh(UStaticMesh* Mesh);
-	void AddFoliageMesh(UFoliageType* Settings);
+	/**  */
+	void OnInstanceCountUpdated(const UFoliageType* FoliageType);
+		
+	/** Whether specified FoliageType can be painted into level */
+	static bool CanPaint(const UFoliageType* FoliageType, const ULevel* InLevel);
 
-	/** Remove a mesh */
-	bool RemoveFoliageMesh(UFoliageType* Settings);
+	/** Add a new asset (FoliageType or StaticMesh) */
+	UFoliageType* AddFoliageAsset(UObject* InAssset);
+	
+	/** Remove a list of Foliage types */
+	bool RemoveFoliageType(UFoliageType** FoliageTypes, int32 Num);
 
 	/** Reapply cluster settings to all the instances */
 	void ReallocateClusters(UFoliageType* Settings);
 
-	/** Replace a mesh with another one */
-	bool ReplaceStaticMesh(UFoliageType* Settings, UStaticMesh* NewStaticMesh, bool& bOutMeshMerged);
-
 	/** Bake meshes to StaticMeshActors */
 	void BakeFoliage(UFoliageType* Settings, bool bSelectedOnly);
 
@@ -293,55 +310,73 @@ public:
 	/** Replace the settings object for this static mesh with the one specified */
 	void ReplaceSettingsObject(UFoliageType* OldSettings, UFoliageType* NewSettings);
 
-	/**
-	 * Save the settings object
-	 *
-	 * @param InSettingsPackageName		The name for the settings package.
-	 * @param StaticMesh				The static mesh to save the settings of.
-	 *
-	 * @return The settings the static mesh is now using.
-	 */
-	UFoliageType* SaveSettingsObject(const FText& InSettingsPackageName, UFoliageType* Settings);
+	/** Save the settings object as an asset */
+	UFoliageType* SaveSettingsObject(UFoliageType* Settings);
 
 	/** Add desired instances. Uses foliage settings to determine location/scale/rotation and whether instances should be ignored */
 	static void AddInstances(UWorld* InWorld, const TArray<FDesiredFoliageInstance>& DesiredInstances);
 
-	/** Apply paint bucket to actor */
-	void ApplyPaintBucket(AActor* Actor, bool bRemove);
-
 	typedef TMap<FName, TMap<ULandscapeComponent*, TArray<uint8> > > LandscapeLayerCacheData;
 private:
 
 	/** Add instances inside the brush to match DesiredInstanceCount */
-	void AddInstancesForBrush(UWorld* InWorld, AInstancedFoliageActor* IFA, UFoliageType* Settings, FFoliageMeshInfo& MeshInfo, int32 DesiredInstanceCount, const TArray<int32>& ExistingInstances, float Pressure);
+	void AddInstancesForBrush(UWorld* InWorld, const UFoliageType* Settings, const FSphere& BrushSphere, int32 DesiredInstanceCount, float Pressure);
+	
+	/** Remove instances inside the brush to match DesiredInstanceCount. */
+	void RemoveInstancesForBrush(UWorld* InWorld, const UFoliageType* Settings, const FSphere& BrushSphere, int32 DesiredInstanceCount, float Pressure);
+
+	/** Apply paint bucket to actor */
+	void ApplyPaintBucket_Add(AActor* Actor);
+	void ApplyPaintBucket_Remove(AActor* Actor);
+
+	/** Reapply instance settings to exiting instances */
+	void ReapplyInstancesDensityForBrush(UWorld* InWorld, const UFoliageType* Settings, const FSphere& BrushSphere, float Pressure);
+	void ReapplyInstancesForBrush(UWorld* InWorld, const UFoliageType* Settings, const FSphere& BrushSphere);
+	void ReapplyInstancesForBrush(UWorld* InWorld, AInstancedFoliageActor* IFA, const UFoliageType* Settings, FFoliageMeshInfo* MeshInfo, const FSphere& BrushSphere);
+
+	/** Select instances inside the brush. */
+	void SelectInstancesForBrush(UWorld* InWorld, const UFoliageType* Settings, const FSphere& BrushSphere, bool bSelect);
+	
+	/** Set/Clear selection for all foliage instances */
+	void SelectInstances(UWorld* InWorld, bool bSelect);
+	
+	/** Set/Clear selection for foliage instances of specific type  */
+	void SelectInstances(UWorld* InWorld, const UFoliageType* Settings, bool bSelect);
+
+	/**  Propagate the selected foliage instances to the actual render components */
+	void ApplySelectionToComponents(UWorld* InWorld, bool bApply);
+
+	/**  Applies relative transformation to selected instances */
+	void TransformSelectedInstances(UWorld* InWorld, const FVector& InDrag, const FRotator& InRot, const FVector& InScale, bool bDuplicate);
+
+	/**  Return selected actor and instance location */
+	AInstancedFoliageActor* GetSelectionLocation(UWorld* InWorld, FVector& OutLocation) const;
+	
+	/**  Updates ed mode widget location to currently selected instance */
+	void UpdateWidgetLocationToInstanceSelection();
+
+	/** Remove currently selected instances*/
+	void RemoveSelectedInstances(UWorld* InWorld);
+			
+	/** Snap instance to the ground   */
+	bool SnapInstanceToGround(AInstancedFoliageActor* InIFA, float AlignMaxAngle, FFoliageMeshInfo& Mesh, int32 InstanceIdx);
+	void SnapSelectedInstancesToGround(UWorld* InWorld);
 
 	/** Common code for adding instances to world based on settings */
-	static void AddInstancesImp(UWorld* InWorld, const UFoliageType* Settings, const TArray<FDesiredFoliageInstance>& DesiredInstances, const AInstancedFoliageActor* IgnoreIFA = nullptr, const TArray<int32>& ExistingInstances = TArray<int32>(), const float Pressure = 1.f, LandscapeLayerCacheData* LandscapeLayerCaches = nullptr, const FFoliageUISettings* UISettings = nullptr);
+	static void AddInstancesImp(UWorld* InWorld, const UFoliageType* Settings, const TArray<FDesiredFoliageInstance>& DesiredInstances, const TArray<int32>& ExistingInstances = TArray<int32>(), const float Pressure = 1.f, LandscapeLayerCacheData* LandscapeLayerCaches = nullptr, const FFoliageUISettings* UISettings = nullptr);
 
 	/** Logic for determining which instances can be placed in the world*/
-	static void CalculatePotentialInstances(const UWorld* InWorld, const AInstancedFoliageActor* IFA, const UFoliageType* Settings, const TArray<FDesiredFoliageInstance>& DesiredInstances, TArray<FPotentialInstance> OutPotentialInstances[NUM_INSTANCE_BUCKETS], LandscapeLayerCacheData* LandscaleLayerCachesPtr, const FFoliageUISettings* UISettings);
+	static void CalculatePotentialInstances(const UWorld* InWorld, const UFoliageType* Settings, const TArray<FDesiredFoliageInstance>& DesiredInstances, TArray<FPotentialInstance> OutPotentialInstances[NUM_INSTANCE_BUCKETS], LandscapeLayerCacheData* LandscaleLayerCachesPtr, const FFoliageUISettings* UISettings);
 
 	/** Similar to CalculatePotentialInstances, but it doesn't do any overlap checks which are much harder to thread. Meant to be run in parallel for placing lots of instances */
-	static void CalculatePotentialInstances_ThreadSafe(const UWorld* InWorld, const AInstancedFoliageActor* IgnoreIFA, const UFoliageType* Settings, const TArray<FDesiredFoliageInstance>* DesiredInstances, TArray<FPotentialInstance> OutPotentialInstances[NUM_INSTANCE_BUCKETS], const FFoliageUISettings* UISettings, const int32 StartIdx, const int32 LastIdx);
-
-	/** Remove instances inside the brush to match DesiredInstanceCount. NOTE: PotentialInstancesToRemove array is modified by this function. */
-	void RemoveInstancesForBrush(AInstancedFoliageActor* IFA, FFoliageMeshInfo& MeshInfo, int32 DesiredInstanceCount, TArray<int32>& PotentialInstancesToRemove, float Pressure);
-
-	/** Reapply instance settings to exiting instances */
-	void ReapplyInstancesForBrush(UWorld* InWorld, AInstancedFoliageActor* IFA, UFoliageType* Settings, const TArray<int32>& ExistingInstances);
-
+	static void CalculatePotentialInstances_ThreadSafe(const UWorld* InWorld, const UFoliageType* Settings, const TArray<FDesiredFoliageInstance>* DesiredInstances, TArray<FPotentialInstance> OutPotentialInstances[NUM_INSTANCE_BUCKETS], const FFoliageUISettings* UISettings, const int32 StartIdx, const int32 LastIdx);
+	
 	/** Lookup the vertex color corresponding to a location traced on a static mesh */
 	static bool GetStaticMeshVertexColorForHit(const UStaticMeshComponent* InStaticMeshComponent, int32 InTriangleIndex, const FVector& InHitLocation, FColor& OutVertexColor);
 
 	/** Does a filter based on the vertex color of a static mesh */
 	static bool VertexMaskCheck(const FHitResult& Hit, const UFoliageType* Settings);
 
-	/** 
-	 * Snaps given instance to the ground  
-	 * @return Whether instance was sucessfully snapped
-	 */
-	bool SnapInstanceToGround(AInstancedFoliageActor* InIFA, float AlignMaxAngle, FFoliageMeshInfo& Mesh, int32 InstanceIdx);
-
 	bool bBrushTraceValid;
 	FVector BrushLocation;
 	FVector BrushTraceDirection;
@@ -350,15 +385,12 @@ private:
 	// Landscape layer cache data
 	LandscapeLayerCacheData LandscapeLayerCaches;
 
-	// Placed level data
-	TArray<FFoliageMeshUIInfo> FoliageMeshList;
-
 	// Cache of instance positions at the start of the transaction
-	TMap<UFoliageType*, FMeshInfoSnapshot> InstanceSnapshot;
+	TMultiMap<UFoliageType*, FMeshInfoSnapshot> InstanceSnapshot;
 
 	bool bToolActive;
 	bool bCanAltDrag;
 
-	// The actor that owns any currently visible selections
-	AInstancedFoliageActor* SelectionIFA;
-};
\ No newline at end of file
+	TArray<FFoliageMeshUIInfoPtr> FoliageMeshList;
+};
+
diff --git a/Engine/Source/Editor/FoliageEdit/Private/FoliageEdModeToolkit.cpp b/Engine/Source/Editor/FoliageEdit/Private/FoliageEdModeToolkit.cpp
--- a/Engine/Source/Editor/FoliageEdit/Private/FoliageEdModeToolkit.cpp
+++ b/Engine/Source/Editor/FoliageEdit/Private/FoliageEdModeToolkit.cpp
@@ -2,7 +2,6 @@
 
 #include "UnrealEd.h"
 #include "FoliageEdModeToolkit.h"
-
 #include "SFoliageEdit.h"
 
 #define LOCTEXT_NAMESPACE "FoliageEditMode"
@@ -46,7 +45,6 @@ TSharedPtr<SWidget> FFoliageEdModeToolkit::GetInlineContent() const
 
 void FFoliageEdModeToolkit::RefreshFullList()
 {
-	// When an undo relates to the Foliage Edit mode, refresh the list.
 	FoliageEdWidget->RefreshFullList();
 }
 
diff --git a/Engine/Source/Editor/FoliageEdit/Private/FoliageEdModeToolkit.h b/Engine/Source/Editor/FoliageEdit/Private/FoliageEdModeToolkit.h
--- a/Engine/Source/Editor/FoliageEdit/Private/FoliageEdModeToolkit.h
+++ b/Engine/Source/Editor/FoliageEdit/Private/FoliageEdModeToolkit.h
@@ -26,4 +26,5 @@ public:
 
 private:
 	TSharedPtr< class SFoliageEdit > FoliageEdWidget;
+	TSharedPtr< class SFoliageEdit2 > FoliageEdWidget2;
 };
\ No newline at end of file
diff --git a/Engine/Source/Editor/FoliageEdit/Private/FoliageEditActions.cpp b/Engine/Source/Editor/FoliageEdit/Private/FoliageEditActions.cpp
--- a/Engine/Source/Editor/FoliageEdit/Private/FoliageEditActions.cpp
+++ b/Engine/Source/Editor/FoliageEdit/Private/FoliageEditActions.cpp
@@ -17,5 +17,8 @@ void FFoliageEditCommands::RegisterCommands()
 	UI_COMMAND( SetNoSettings, "Hide Details", "Hide details.", EUserInterfaceActionType::ToggleButton, FInputGesture() );
 	UI_COMMAND( SetPaintSettings, "Show Painting settings", "Show painting settings.", EUserInterfaceActionType::ToggleButton, FInputGesture() );
 	UI_COMMAND( SetClusterSettings, "Show Instance settings", "Show settings for placed instances.", EUserInterfaceActionType::ToggleButton, FInputGesture() );
+
+	UI_COMMAND( RemoveFoliageType, "Remove", "Remove all associated foliage instances", EUserInterfaceActionType::Button, FInputGesture() );
+	UI_COMMAND( ShowFoliageTypeInCB, "Show in Content Browser", "Show asset in Content Browser.", EUserInterfaceActionType::Button, FInputGesture() );
 }
 
diff --git a/Engine/Source/Editor/FoliageEdit/Private/FoliageEditActions.h b/Engine/Source/Editor/FoliageEdit/Private/FoliageEditActions.h
--- a/Engine/Source/Editor/FoliageEdit/Private/FoliageEditActions.h
+++ b/Engine/Source/Editor/FoliageEdit/Private/FoliageEditActions.h
@@ -37,6 +37,10 @@ public:
 	TSharedPtr< FUICommandInfo > SetPaintSettings;
 	TSharedPtr< FUICommandInfo > SetClusterSettings;
 
+	/** FoliageType commands */
+	TSharedPtr< FUICommandInfo > RemoveFoliageType;
+	TSharedPtr< FUICommandInfo > ShowFoliageTypeInCB;
+
 	/**
 	 * Initialize commands
 	 */
diff --git a/Engine/Source/Editor/FoliageEdit/Private/SFoliageEdit.cpp b/Engine/Source/Editor/FoliageEdit/Private/SFoliageEdit.cpp
--- a/Engine/Source/Editor/FoliageEdit/Private/SFoliageEdit.cpp
+++ b/Engine/Source/Editor/FoliageEdit/Private/SFoliageEdit.cpp
@@ -89,8 +89,6 @@ void SFoliageEdit::Construct(const FArguments& InArgs)
 
 	BindCommands();
 
-	FEditorDelegates::NewCurrentLevel.AddSP(this, &SFoliageEdit::NotifyNewCurrentLevel);
-
 	AssetThumbnailPool = MakeShareable(new FAssetThumbnailPool(512, TAttribute<bool>::Create(TAttribute<bool>::FGetter::CreateSP(this, &SFoliageEdit::IsHovered))));
 
 	// Everything (or almost) uses this padding, change it to expand the padding.
@@ -351,8 +349,9 @@ void SFoliageEdit::RefreshFullList()
 	{
 		RemoveItemFromScrollbox(DisplayItemList[DisplayItemIdx]);
 	}
-
-	TArray<struct FFoliageMeshUIInfo>& FoliageMeshList = FoliageEditMode->GetFoliageMeshList();
+	
+	//FoliageEditMode->PopulateFoliageMeshList();
+	auto& FoliageMeshList = FoliageEditMode->GetFoliageMeshList();
 
 	for (int MeshIdx = 0; MeshIdx < FoliageMeshList.Num(); MeshIdx++)
 	{
@@ -360,16 +359,16 @@ void SFoliageEdit::RefreshFullList()
 	}
 }
 
-void SFoliageEdit::AddItemToScrollbox(struct FFoliageMeshUIInfo& InFoliageInfoToAdd)
+void SFoliageEdit::AddItemToScrollbox(TSharedPtr<FFoliageMeshUIInfo>& InFoliageInfoToAdd)
 {
-	UFoliageType* Settings = InFoliageInfoToAdd.Settings;
+	UFoliageType* Settings = InFoliageInfoToAdd->Settings;
 
 	TSharedRef<SFoliageEditMeshDisplayItem> DisplayItem =
 		SNew(SFoliageEditMeshDisplayItem)
 		.FoliageEditPtr(SharedThis(this))
 		.FoliageSettingsPtr(Settings)
 		.AssetThumbnail(CreateThumbnail(Settings->GetStaticMesh()))
-		.FoliageMeshUIInfo(MakeShareable(new FFoliageMeshUIInfo(InFoliageInfoToAdd)));
+		.FoliageMeshUIInfo(InFoliageInfoToAdd);
 
 	DisplayItemList.Add(DisplayItem);
 
@@ -388,17 +387,17 @@ void SFoliageEdit::RemoveItemFromScrollbox(const TSharedPtr<SFoliageEditMeshDisp
 
 void SFoliageEdit::ReplaceItem(const TSharedPtr<SFoliageEditMeshDisplayItem> InDisplayItemToReplaceIn, UStaticMesh* InNewStaticMesh)
 {
-	TArray<FFoliageMeshUIInfo>& FoliageMeshList = FoliageEditMode->GetFoliageMeshList();
+	TArray<FFoliageMeshUIInfoPtr>& FoliageMeshList = FoliageEditMode->GetFoliageMeshList();
 
-	for (FFoliageMeshUIInfo& UIInfo : FoliageMeshList)
+	for (FFoliageMeshUIInfoPtr& UIInfo : FoliageMeshList)
 	{
-		UFoliageType* Settings = UIInfo.Settings;
+		UFoliageType* Settings = UIInfo->Settings;
 		if (Settings->GetStaticMesh() == InNewStaticMesh)
 		{
 			// This is the info for the new static mesh. Update the display item.
 			InDisplayItemToReplaceIn->Replace(Settings,
 				CreateThumbnail(Settings->GetStaticMesh()),
-				MakeShareable(new FFoliageMeshUIInfo(UIInfo)));
+				UIInfo);
 			break;
 		}
 	}
@@ -609,7 +608,7 @@ EVisibility SFoliageEdit::GetVisibility_FoliageDropTarget() const
 		TArray<FAssetData> DraggedAssets = AssetUtil::ExtractAssetDataFromDrag(FSlateApplication::Get().GetDragDroppingContent());
 		for ( const FAssetData& AssetData : DraggedAssets )
 		{
-			if ( AssetData.IsValid() && AssetData.GetClass()->IsChildOf(UStaticMesh::StaticClass()) )
+			if ( AssetData.IsValid() && (AssetData.GetClass()->IsChildOf(UStaticMesh::StaticClass()) || AssetData.GetClass()->IsChildOf(UFoliageType::StaticClass())) )
 			{
 				return EVisibility::Visible;
 			}
@@ -631,9 +630,11 @@ EVisibility SFoliageEdit::GetVisibility_NonEmptyList() const
 
 bool SFoliageEdit::CanAddStaticMesh(const UStaticMesh* const InStaticMesh) const
 {
-	for (FFoliageMeshUIInfo& UIInfo : FoliageEditMode->GetFoliageMeshList())
+	TArray<FFoliageMeshUIInfoPtr>& FoliageMeshList = FoliageEditMode->GetFoliageMeshList();
+	
+	for (FFoliageMeshUIInfoPtr& UIInfo : FoliageMeshList)
 	{
-		UFoliageType* Settings = UIInfo.Settings;
+		UFoliageType* Settings = UIInfo->Settings;
 
 		if (Settings->GetStaticMesh() == InStaticMesh)
 		{
@@ -650,15 +651,13 @@ FReply SFoliageEdit::OnDrop_ListView(const FGeometry& MyGeometry, const FDragDro
 	for (int32 AssetIdx = 0; AssetIdx < DroppedAssetData.Num(); ++AssetIdx)
 	{
 		GWarn->BeginSlowTask(LOCTEXT("OnDrop_LoadPackage", "Fully Loading Package For Drop"), true, false);
-		UStaticMesh* StaticMesh = Cast<UStaticMesh>(DroppedAssetData[AssetIdx].GetAsset());
+		UObject* Asset = DroppedAssetData[AssetIdx].GetAsset();
 		GWarn->EndSlowTask();
 
-		if (StaticMesh && CanAddStaticMesh(StaticMesh))
+		if (Asset)
 		{
-			FoliageEditMode->AddFoliageMesh(StaticMesh);
-			AddItemToScrollbox(FoliageEditMode->GetFoliageMeshList()[FoliageEditMode->GetFoliageMeshList().Num() - 1]);
+			FoliageEditMode->AddFoliageAsset(Asset);
 		}
-
 	}
 
 	return FReply::Handled();
@@ -744,11 +743,4 @@ EVisibility SFoliageEdit::GetVisibility_Filters() const
 	return EVisibility::Visible;
 }
 
-
-void SFoliageEdit::NotifyNewCurrentLevel()
-{
-	RefreshFullList();
-}
-
-
 #undef LOCTEXT_NAMESPACE
diff --git a/Engine/Source/Editor/FoliageEdit/Private/SFoliageEdit.h b/Engine/Source/Editor/FoliageEdit/Private/SFoliageEdit.h
--- a/Engine/Source/Editor/FoliageEdit/Private/SFoliageEdit.h
+++ b/Engine/Source/Editor/FoliageEdit/Private/SFoliageEdit.h
@@ -27,7 +27,7 @@ public:
 	void RefreshFullList();
 
 	/** Adds a static mesh to the list of available meshes for foliage. May be called on by the MeshListView items. */
-	void AddItemToScrollbox(FFoliageMeshUIInfo& InFoliageInfoToAdd);
+	void AddItemToScrollbox(TSharedPtr<FFoliageMeshUIInfo>& InFoliageInfoToAdd);
 
 	/** Removes a static mesh from the list of available meshes for foliage. May be called on by the MeshListView items. */
 	void RemoveItemFromScrollbox(const TSharedPtr<SFoliageEditMeshDisplayItem> InWidgetToRemove);
@@ -41,7 +41,6 @@ public:
 	class FEdModeFoliage* GetFoliageEditMode() const { return FoliageEditMode; }
 
 private:
-
 	/** Clears all the tools selection by setting them to false. */
 	void ClearAllToolSelection();
 
@@ -149,7 +148,6 @@ private:
 	/** Checks if the filters should appear. Dependant on the current tool being used. */
 	EVisibility GetVisibility_Filters() const;
 
-
 	/**
 	 * Checks if a static mesh can be added to the list of Static Meshes available.
 	 *
@@ -158,18 +156,10 @@ private:
 	 * @return	Returns true if the static mesh is not currently in the list, false if it is.
 	 */
 	bool CanAddStaticMesh(const UStaticMesh* const InStaticMesh) const;
-
-	/** Callback handler when user changes the active streaming level with the level browser */
-	void NotifyNewCurrentLevel();
-
 private:
 	/** The list view object for displaying Static Meshes to use for foliage. */
 	TSharedPtr<SListView<TSharedPtr<FFoliageMeshUIInfo>>> MeshListView;
-
-	/** List of Static Meshes being used for foliage, adapted from a list retrieved from Foliage Mode and items should only be added to it in NotifyChanged to keep information
-		accurate between it and the Foliage Mode. */
-	TArray<TSharedPtr<FFoliageMeshUIInfo>> MeshList;
-
+	
 	/** Pool for maintaining and rendering thumbnails */
 	TSharedPtr<FAssetThumbnailPool> AssetThumbnailPool;
 
diff --git a/Engine/Source/Editor/FoliageEdit/Private/SFoliageEditMeshDisplayItem.cpp b/Engine/Source/Editor/FoliageEdit/Private/SFoliageEditMeshDisplayItem.cpp
--- a/Engine/Source/Editor/FoliageEdit/Private/SFoliageEditMeshDisplayItem.cpp
+++ b/Engine/Source/Editor/FoliageEdit/Private/SFoliageEditMeshDisplayItem.cpp
@@ -1536,7 +1536,7 @@ void SFoliageEditMeshDisplayItem::Construct(const FArguments& InArgs)
 						SNew(SButton)
 						.ButtonStyle(FEditorStyle::Get(), "NoBorder")
 						.OnClicked(this, &SFoliageEditMeshDisplayItem::OnReplace)
-						.ToolTipText(NSLOCTEXT("FoliageEdMode", "Replace_Tooltip", "Replace all instances with the Static Mesh currently selected in the Content Browser."))
+						.ToolTipText(NSLOCTEXT("FoliageEdMode", "Replace_Tooltip", "Replace with Static Mesh currently selected in the Content Browser."))
 						.AddMetaData<FTutorialMetaData>(FTutorialMetaData(TEXT("Foliage.ReplaceInstances"), "LevelEditorToolbox"))
 						[
 							SNew(SImage)
@@ -1551,7 +1551,7 @@ void SFoliageEditMeshDisplayItem::Construct(const FArguments& InArgs)
 						SNew(SButton)
 						.ButtonStyle(FEditorStyle::Get(), "NoBorder")
 						.OnClicked(this, &SFoliageEditMeshDisplayItem::OnSync)
-						.ToolTipText(NSLOCTEXT("FoliageEdMode", "FindInContentBrowser_Tooltip", "Find this Static Mesh in the Content Browser."))
+						.ToolTipText(NSLOCTEXT("FoliageEdMode", "FindInContentBrowser_Tooltip", "Find this Asset in the Content Browser."))
 						.AddMetaData<FTutorialMetaData>(FTutorialMetaData(TEXT("Foliage.FindInBrowser"), "LevelEditorToolbox"))
 						[
 							SNew(SImage)
@@ -1610,8 +1610,9 @@ void SFoliageEditMeshDisplayItem::Construct(const FArguments& InArgs)
 					[
 						SNew(SButton)
 						.ButtonStyle(FEditorStyle::Get(), "NoBorder")
-						.OnClicked(this, &SFoliageEditMeshDisplayItem::OnSaveRemoveSettings)
-						.ToolTipText(this, &SFoliageEditMeshDisplayItem::GetSaveRemoveSettingsTooltip)
+						.OnClicked(this, &SFoliageEditMeshDisplayItem::OnSaveSettings)
+						.Visibility(this, &SFoliageEditMeshDisplayItem::IsSaveSettingsVisible)
+						.ToolTipText(this, &SFoliageEditMeshDisplayItem::GetSaveSettingsTooltip)
 						.AddMetaData<FTutorialMetaData>(FTutorialMetaData(TEXT("Foliage.SaveSettings"), "LevelEditorToolbox"))
 						[
 							SNew(SImage)
@@ -1816,23 +1817,11 @@ FReply SFoliageEditMeshDisplayItem::OnReplace()
 	FEditorDelegates::LoadSelectedAssetsIfNeeded.Broadcast();
 	USelection* SelectedSet = GEditor->GetSelectedSet(UStaticMesh::StaticClass());
 	UStaticMesh* SelectedStaticMesh = Cast<UStaticMesh>(SelectedSet->GetTop(UStaticMesh::StaticClass()));
-	if (SelectedStaticMesh != NULL)
+	if (SelectedStaticMesh != nullptr)
 	{
-		FEdModeFoliage* Mode = (FEdModeFoliage*)GLevelEditorModeTools().GetActiveMode(FBuiltinEditorModes::EM_Foliage);
-
-		bool bMeshMerged = false;
-		if (Mode->ReplaceStaticMesh(FoliageSettingsPtr, SelectedStaticMesh, bMeshMerged))
-		{
-			// If they were merged, simply remove the current item. Otherwise replace it.
-			if (bMeshMerged)
-			{
-				FoliageEditPtr.Pin()->RemoveItemFromScrollbox(SharedThis(this));
-			}
-			else
-			{
-				FoliageEditPtr.Pin()->ReplaceItem(SharedThis(this), SelectedStaticMesh);
-			}
-		}
+		FoliageSettingsPtr->Modify();
+		FoliageSettingsPtr->SetStaticMesh(SelectedStaticMesh);
+		FoliageSettingsPtr->PostEditChange();
 	}
 
 	return FReply::Handled();
@@ -1861,9 +1850,15 @@ void SFoliageEditMeshDisplayItem::OnSelectionChanged(ECheckBoxState InType)
 FReply SFoliageEditMeshDisplayItem::OnSync()
 {
 	TArray<UObject*> Objects;
-
-	Objects.Add(FoliageSettingsPtr->GetStaticMesh());
-
+	if (FoliageSettingsPtr->IsAsset())
+	{
+		Objects.Add(FoliageSettingsPtr);
+	}
+	else
+	{
+		Objects.Add(FoliageSettingsPtr->GetStaticMesh());
+	}
+	
 	GEditor->SyncBrowserToObjects(Objects);
 
 	return FReply::Handled();
@@ -1873,7 +1868,7 @@ FReply SFoliageEditMeshDisplayItem::OnRemove()
 {
 	FEdModeFoliage* Mode = (FEdModeFoliage*)GLevelEditorModeTools().GetActiveMode(FBuiltinEditorModes::EM_Foliage);
 
-	if (Mode->RemoveFoliageMesh(FoliageSettingsPtr))
+	if (Mode->RemoveFoliageType(&FoliageSettingsPtr, 1))
 	{
 		FoliageEditPtr.Pin()->RemoveItemFromScrollbox(SharedThis(this));
 	}
@@ -1881,44 +1876,12 @@ FReply SFoliageEditMeshDisplayItem::OnRemove()
 	return FReply::Handled();
 }
 
-FReply SFoliageEditMeshDisplayItem::OnSaveRemoveSettings()
+FReply SFoliageEditMeshDisplayItem::OnSaveSettings()
 {
 	FEdModeFoliage* Mode = (FEdModeFoliage*)GLevelEditorModeTools().GetActiveMode(FBuiltinEditorModes::EM_Foliage);
 
-	if (FoliageSettingsPtr->GetOuter()->IsA(UPackage::StaticClass()))
-	{
-		UFoliageType* NewSettings = NULL;
-		NewSettings = Mode->CopySettingsObject(FoliageSettingsPtr);
-
-		// Do not replace the current one if NULL is returned, just keep the old one.
-		if (NewSettings)
-		{
-			FoliageSettingsPtr = NewSettings;
-		}
-	}
-	else
-	{
-		// Build default settings asset name and path
-		FString DefaultAsset = FPackageName::GetLongPackagePath(FoliageSettingsPtr->GetStaticMesh()->GetOutermost()->GetName()) + TEXT("/") + FoliageSettingsPtr->GetStaticMesh()->GetName() + TEXT("_settings");
-
-		TSharedRef<SDlgPickAssetPath> SettingDlg =
-			SNew(SDlgPickAssetPath)
-			.Title(LOCTEXT("SettingsDialogTitle", "Choose Location for Foliage Settings Asset"))
-			.DefaultAssetPath(FText::FromString(DefaultAsset));
-
-		if (SettingDlg->ShowModal() != EAppReturnType::Cancel)
-		{
-			UFoliageType* NewSettings = NULL;
-
-			NewSettings = Mode->SaveSettingsObject(SettingDlg->GetFullAssetPath(), FoliageSettingsPtr);
-
-			// Do not replace the current one if NULL is returned, just keep the old one.
-			if (NewSettings)
-			{
-				FoliageSettingsPtr = NewSettings;
-			}
-		}
-	}
+	UFoliageType* NewSettings = Mode->SaveSettingsObject(FoliageSettingsPtr);
+	Mode->PopulateFoliageMeshList();
 
 	return FReply::Handled();
 }
@@ -1938,16 +1901,19 @@ FReply SFoliageEditMeshDisplayItem::OnOpenSettings()
 	return FReply::Handled();
 }
 
-FText SFoliageEditMeshDisplayItem::GetSaveRemoveSettingsTooltip() const
+FText SFoliageEditMeshDisplayItem::GetSaveSettingsTooltip() const
 {
-	// Remove Settings tooltip.
-	if (FoliageSettingsPtr->GetOuter()->IsA(UPackage::StaticClass()))
+	if (!FoliageSettingsPtr->IsAsset())
 	{
-		return NSLOCTEXT("FoliageEdMode", "RemoveSettings_Tooltip", "Do not store the foliage settings in a shared InstancedFoliageSettings object.");
+		return NSLOCTEXT("FoliageEdMode", "SaveSettings_Tooltip", "Save these settings as an shared Asset");
 	}
 
-	// Save settings tooltip.
-	return NSLOCTEXT("FoliageEdMode", "SaveSettings_Tooltip", "Save these settings as an InstancedFoliageSettings object stored in a package.");
+	return FText();
+}
+
+EVisibility SFoliageEditMeshDisplayItem::IsSaveSettingsVisible() const
+{
+	return !FoliageSettingsPtr->IsAsset() ? EVisibility::Visible : EVisibility::Collapsed;
 }
 
 bool SFoliageEditMeshDisplayItem::IsPropertyVisible(const FPropertyAndParent& PropertyAndParent) const
@@ -1968,7 +1934,17 @@ bool SFoliageEditMeshDisplayItem::IsPropertyVisible(const FPropertyAndParent& Pr
 
 FText SFoliageEditMeshDisplayItem::GetInstanceCountString() const
 {
-	return FText::Format(LOCTEXT("InstanceCount_Value", "Instance Count: {0}"), FText::AsNumber(FoliageMeshUIInfo->MeshInfo ? FoliageMeshUIInfo->MeshInfo->GetInstanceCount() : 0 ));
+	const int32 InsatanceCountTotal = FoliageMeshUIInfo->InstanceCountTotal;
+	const int32 InstanceCountCurrentLevel = FoliageMeshUIInfo->InstanceCountCurrentLevel;
+	
+	if (InsatanceCountTotal != InstanceCountCurrentLevel)
+	{
+		return FText::Format(LOCTEXT("InstanceCount_ValueT", "Instance Count: {0} ({1})"), FText::AsNumber(InstanceCountCurrentLevel), FText::AsNumber(InsatanceCountTotal));
+	}
+	else
+	{
+		return FText::Format(LOCTEXT("InstanceCount_ValueC", "Instance Count: {0}"), FText::AsNumber(InsatanceCountTotal));
+	}
 }
 
 EVisibility SFoliageEditMeshDisplayItem::IsReapplySettingsVisible() const
@@ -2573,11 +2549,6 @@ FReply SFoliageEditMeshDisplayItem::OnMouseDownSelection(const FGeometry& MyGeom
 
 const FSlateBrush* SFoliageEditMeshDisplayItem::GetSaveSettingsBrush() const
 {
-	if (FoliageSettingsPtr->GetOuter()->IsA(UPackage::StaticClass()))
-	{
-		return FEditorStyle::GetBrush(TEXT("FoliageEditMode.DeleteItem"));
-	}
-
 	return FEditorStyle::GetBrush(TEXT("FoliageEditMode.SaveSettings"));
 }
 
diff --git a/Engine/Source/Editor/FoliageEdit/Private/SFoliageEditMeshDisplayItem.h b/Engine/Source/Editor/FoliageEdit/Private/SFoliageEditMeshDisplayItem.h
--- a/Engine/Source/Editor/FoliageEdit/Private/SFoliageEditMeshDisplayItem.h
+++ b/Engine/Source/Editor/FoliageEdit/Private/SFoliageEditMeshDisplayItem.h
@@ -58,7 +58,7 @@ public:
 	/** Handles clicking on the selection area for the item. */
 	FReply OnMouseDownSelection(const FGeometry& MyGeometry, const FPointerEvent& MouseEvent);
 
-	/** Returns the current brush for the Save/Remove settings button. */
+	/** Returns the current brush for the Save settings button. */
 	const FSlateBrush* GetSaveSettingsBrush() const;
 
 private:
@@ -106,12 +106,13 @@ private:
 
 	FReply OnRemove();
 
-	FReply OnSaveRemoveSettings();
+	FReply OnSaveSettings();
 
 	FReply OnOpenSettings();
 
 	/** Returns the tool-tip to use for the save settings button. */
-	FText GetSaveRemoveSettingsTooltip() const;
+	FText GetSaveSettingsTooltip() const;
+	EVisibility IsSaveSettingsVisible() const;
 
 	/** Determines what properties to display in the details panel. */
 	bool IsPropertyVisible(const FPropertyAndParent& PropertyAndParent) const;
diff --git a/Engine/Source/Editor/LandscapeEditor/Private/LandscapeEdMode.cpp b/Engine/Source/Editor/LandscapeEditor/Private/LandscapeEdMode.cpp
--- a/Engine/Source/Editor/LandscapeEditor/Private/LandscapeEdMode.cpp
+++ b/Engine/Source/Editor/LandscapeEditor/Private/LandscapeEdMode.cpp
@@ -2620,7 +2620,7 @@ ALandscape* FEdModeLandscape::ChangeComponentSetting(int32 NumComponentsX, int32
 							if (OldFoliageActor)
 							{
 								OldFoliageActor->MoveInstancesToNewComponent(OldCollisionComponent, NewCollisionComponent);
-								AInstancedFoliageActor* NewFoliageActor = AInstancedFoliageActor::GetInstancedFoliageActorForLevel(NewCollisionComponent->GetTypedOuter<ULevel>());
+								AInstancedFoliageActor* NewFoliageActor = AInstancedFoliageActor::GetInstancedFoliageActorForLevel(NewCollisionComponent->GetTypedOuter<ULevel>(), true);
 								NewCollisionComponent->SnapFoliageInstances(*NewFoliageActor, FBox(FVector(-WORLD_MAX), FVector(WORLD_MAX)));
 							}
 						}
diff --git a/Engine/Source/Editor/LandscapeEditor/Private/LandscapeEdModeTools.h b/Engine/Source/Editor/LandscapeEditor/Private/LandscapeEdModeTools.h
--- a/Engine/Source/Editor/LandscapeEditor/Private/LandscapeEdModeTools.h
+++ b/Engine/Source/Editor/LandscapeEditor/Private/LandscapeEdModeTools.h
@@ -829,7 +829,10 @@ struct FXYOffsetmapAccessor
 				{
 					ULandscapeHeightfieldCollisionComponent* CollisionComponent = CollisionComponents[Index];
 					AInstancedFoliageActor* IFA = AInstancedFoliageActor::GetInstancedFoliageActorForLevel(CollisionComponent->GetComponentLevel());
-					CollisionComponent->SnapFoliageInstances(*IFA, PreUpdateLocalBoxes[Index].TransformBy(LandscapeInfo->GetLandscapeProxy()->LandscapeActorToWorld().ToMatrixWithScale()).ExpandBy(1.0f));
+					if (IFA)
+					{
+						CollisionComponent->SnapFoliageInstances(*IFA, PreUpdateLocalBoxes[Index].TransformBy(LandscapeInfo->GetLandscapeProxy()->LandscapeActorToWorld().ToMatrixWithScale()).ExpandBy(1.0f));
+					}
 				}
 			}
 			else
diff --git a/Engine/Source/Editor/UnrealEd/Private/EditorObject.cpp b/Engine/Source/Editor/UnrealEd/Private/EditorObject.cpp
--- a/Engine/Source/Editor/UnrealEd/Private/EditorObject.cpp
+++ b/Engine/Source/Editor/UnrealEd/Private/EditorObject.cpp
@@ -251,7 +251,7 @@ static const TCHAR* ImportProperties(
 					ULevel* ComponentLevel = CastChecked<ULevel>(SubobjectRoot->GetOuter());
 					if (ComponentLevel->IsCurrentLevel())
 					{
-						AInstancedFoliageActor* IFA = AInstancedFoliageActor::GetInstancedFoliageActorForCurrentLevel(ComponentLevel->OwningWorld);
+						AInstancedFoliageActor* IFA = AInstancedFoliageActor::GetInstancedFoliageActorForCurrentLevel(ComponentLevel->OwningWorld, true);
 
 						const TCHAR* StrPtr;
 						FString TextLine;
diff --git a/Engine/Source/Runtime/Engine/Classes/Engine/Level.h b/Engine/Source/Runtime/Engine/Classes/Engine/Level.h
--- a/Engine/Source/Runtime/Engine/Classes/Engine/Level.h
+++ b/Engine/Source/Runtime/Engine/Classes/Engine/Level.h
@@ -7,6 +7,7 @@
 class ALevelBounds;
 class UTexture2D;
 class UNavigationDataChunk;
+class AInstancedFoliageActor;
 
 /**
  * Structure containing all information needed for determining the screen space
@@ -456,6 +457,9 @@ public:
 	/** Actor which defines level logical bounding box				*/
 	TWeakObjectPtr<ALevelBounds>				LevelBoundsActor;
 
+	/** Cached pointer to Foliage actor		*/
+	TWeakObjectPtr<AInstancedFoliageActor>		InstancedFoliageActor;
+
 	/** Called when Level bounds actor has been updated */
 	DECLARE_EVENT( ULevel, FLevelBoundsActorUpdatedEvent );
 	FLevelBoundsActorUpdatedEvent& LevelBoundsActorUpdated() { return LevelBoundsActorUpdatedEvent; }
diff --git a/Engine/Source/Runtime/Foliage/Private/FoliageStatistics.cpp b/Engine/Source/Runtime/Foliage/Private/FoliageStatistics.cpp
--- a/Engine/Source/Runtime/Foliage/Private/FoliageStatistics.cpp
+++ b/Engine/Source/Runtime/Foliage/Private/FoliageStatistics.cpp
@@ -21,8 +21,11 @@ int32 UFoliageStatistics::FoliageOverlappingSphereCount(UObject* WorldContextObj
 	{
 		ULevel* Level = *LvlItr;
 		AInstancedFoliageActor* IFA = AInstancedFoliageActor::GetInstancedFoliageActorForLevel(Level);
-		const UFoliageType* FoliageType = IFA->GetSettingsForMesh(Mesh);
-		Count += IFA->GetOverlappingSphereCount(FoliageType, Sphere);
+		if (IFA)
+		{
+			const UFoliageType* FoliageType = IFA->GetSettingsForMesh(Mesh);
+			Count += IFA->GetOverlappingSphereCount(FoliageType, Sphere);
+		}
 	}
 
 	return Count;
@@ -37,9 +40,11 @@ int32 UFoliageStatistics::FoliageOverlappingBoxCount(UObject* WorldContextObject
 	{
 		ULevel* Level = *LvlItr;
 		AInstancedFoliageActor* IFA = AInstancedFoliageActor::GetInstancedFoliageActorForLevel(Level);
-		const UFoliageType* Type = IFA->GetSettingsForMesh(StaticMesh);
-
-		Count += IFA->GetOverlappingBoxCount(Type, Box);
+		if (IFA)
+		{
+			const UFoliageType* Type = IFA->GetSettingsForMesh(StaticMesh);
+			Count += IFA->GetOverlappingBoxCount(Type, Box);
+		}
 	}
 
 	return Count;
@@ -53,8 +58,10 @@ void UFoliageStatistics::FoliageOverlappingBoxTransforms(UObject* WorldContextOb
 	{
 		ULevel* Level = *LvlItr;
 		AInstancedFoliageActor* IFA = AInstancedFoliageActor::GetInstancedFoliageActorForLevel(Level);
-		const UFoliageType* Type = IFA->GetSettingsForMesh(StaticMesh);
-
-		IFA->GetOverlappingBoxTransforms(Type, Box, OutTransforms);
+		if (IFA)
+		{
+			const UFoliageType* Type = IFA->GetSettingsForMesh(StaticMesh);
+			IFA->GetOverlappingBoxTransforms(Type, Box, OutTransforms);
+		}
 	}
 }
diff --git a/Engine/Source/Runtime/Foliage/Private/InstancedFoliage.cpp b/Engine/Source/Runtime/Foliage/Private/InstancedFoliage.cpp
--- a/Engine/Source/Runtime/Foliage/Private/InstancedFoliage.cpp
+++ b/Engine/Source/Runtime/Foliage/Private/InstancedFoliage.cpp
@@ -630,9 +630,7 @@ void FFoliageMeshInfo::RemoveInstances(AInstancedFoliageActor* InIFA, const TArr
 			// Remove the removed item from the removal list
 			InstancesToRemove.Remove(InstanceIndexToRemove);
 		}
-
-		InIFA->CheckSelection();
-
+			
 		CheckValid();
 	}
 }
@@ -847,6 +845,31 @@ bool FFoliageMeshInfo::CheckForOverlappingInstanceExcluding(int32 TestInstanceId
 	return false;
 }
 
+void FFoliageMeshInfo::SelectInstances(AInstancedFoliageActor* InIFA, bool bSelect)
+{
+	if (Component)
+	{
+		InIFA->Modify();
+	
+		if (bSelect)
+		{
+			for (int32 i = 0; i < Component->PerInstanceSMData.Num(); ++i)
+			{
+				SelectedIndices.Add(i);
+			}
+		}
+		else
+		{
+			SelectedIndices.Empty();
+		}
+
+		// Apply selections to the component
+		Component->SelectedInstances.Init(bSelect, Component->PerInstanceSMData.Num());
+		Component->ReleasePerInstanceRenderData();
+		Component->MarkRenderStateDirty();
+	}
+}
+
 void FFoliageMeshInfo::SelectInstances(AInstancedFoliageActor* InIFA, bool bSelect, TArray<int32>& InInstances)
 {
 	if (InInstances.Num())
@@ -909,47 +932,26 @@ AInstancedFoliageActor::AInstancedFoliageActor(const FObjectInitializer& ObjectI
 
 AInstancedFoliageActor* AInstancedFoliageActor::GetInstancedFoliageActorForCurrentLevel(UWorld* InWorld, bool bCreateIfNone)
 {
-	for (TActorIterator<AInstancedFoliageActor> It(InWorld); It; ++It)
-	{
-		AInstancedFoliageActor* InstancedFoliageActor = *It;
-		if (InstancedFoliageActor)
-		{
-			if (InstancedFoliageActor->GetLevel()->IsCurrentLevel() && !InstancedFoliageActor->IsPendingKill())
-			{
-				return InstancedFoliageActor;
-			}
-		}
-	}
-
-	return bCreateIfNone ? InWorld->SpawnActor<AInstancedFoliageActor>() : nullptr;
+	return GetInstancedFoliageActorForLevel(InWorld->GetCurrentLevel(), bCreateIfNone);
 }
 
-
-AInstancedFoliageActor* AInstancedFoliageActor::GetInstancedFoliageActorForLevel(ULevel* InLevel, bool bCreateIfNone /* = true */)
+AInstancedFoliageActor* AInstancedFoliageActor::GetInstancedFoliageActorForLevel(ULevel* InLevel, bool bCreateIfNone /* = false */)
 {
-	if (!InLevel)
+	AInstancedFoliageActor* IFA = nullptr;
+	if (InLevel)
 	{
-		return nullptr;
-	}
-
-	for (int32 ActorIndex = 0; ActorIndex < InLevel->Actors.Num(); ++ActorIndex)
-	{
-		AInstancedFoliageActor* InstancedFoliageActor = Cast<AInstancedFoliageActor>(InLevel->Actors[ActorIndex]);
-		if (InstancedFoliageActor && !InstancedFoliageActor->IsPendingKill())
+		IFA = InLevel->InstancedFoliageActor.Get();
+		
+		if (!IFA && bCreateIfNone)
 		{
-			return InstancedFoliageActor;
+			FActorSpawnParameters SpawnParams;
+			SpawnParams.OverrideLevel = InLevel;
+			IFA = InLevel->GetWorld()->SpawnActor<AInstancedFoliageActor>(SpawnParams);
+			InLevel->InstancedFoliageActor = IFA;
 		}
 	}
 
-	if (bCreateIfNone)
-	{
-		FActorSpawnParameters SpawnParams;
-		SpawnParams.OverrideLevel = InLevel;
-
-		return InLevel->GetWorld()->SpawnActor<AInstancedFoliageActor>(SpawnParams);
-	}
-
-	return nullptr;
+	return IFA;
 }
 
 
@@ -1095,21 +1097,40 @@ void AInstancedFoliageActor::MoveInstancesForMovedComponent(UActorComponent* InC
 
 void AInstancedFoliageActor::DeleteInstancesForComponent(UActorComponent* InComponent)
 {
+	const auto BaseId = InstanceBaseCache.GetInstanceBaseId(InComponent);
+	// Instances with empty base has BaseId==InvalidBaseId, we should not delete these
+	if (BaseId == FFoliageInstanceBaseCache::InvalidBaseId)
+	{
+		return;
+	}
+		
 	for (auto& MeshPair : FoliageMeshes)
 	{
 		FFoliageMeshInfo& MeshInfo = *MeshPair.Value;
-		const auto BaseId = InstanceBaseCache.GetInstanceBaseId(InComponent);
-
-		// Instances with empty base has BaseId==InvalidBaseId, we should not delete these
-		if (BaseId == FFoliageInstanceBaseCache::InvalidBaseId)
+		const auto* InstanceSet = MeshInfo.ComponentHash.Find(BaseId);
+		if (InstanceSet)
 		{
-			continue;
+			MeshInfo.RemoveInstances(this, InstanceSet->Array());
 		}
+	}
+}
 
-		const auto* InstanceSet = MeshInfo.ComponentHash.Find(BaseId);
+void AInstancedFoliageActor::DeleteInstancesForComponent(UActorComponent* InComponent, const UFoliageType* FoliageType)
+{
+	const auto BaseId = InstanceBaseCache.GetInstanceBaseId(InComponent);
+	// Instances with empty base has BaseId==InvalidBaseId, we should not delete these
+	if (BaseId == FFoliageInstanceBaseCache::InvalidBaseId)
+	{
+		return;
+	}
+	
+	FFoliageMeshInfo* MeshInfo = FindMesh(FoliageType);
+	if (MeshInfo)
+	{
+		const auto* InstanceSet = MeshInfo->ComponentHash.Find(BaseId);
 		if (InstanceSet)
 		{
-			MeshInfo.RemoveInstances(this, InstanceSet->Array());
+			MeshInfo->RemoveInstances(this, InstanceSet->Array());
 		}
 	}
 }
@@ -1138,7 +1159,7 @@ void AInstancedFoliageActor::DeleteInstancesForProceduralFoliageComponent(const
 
 void AInstancedFoliageActor::MoveInstancesForComponentToCurrentLevel(UActorComponent* InComponent)
 {
-	AInstancedFoliageActor* NewIFA = AInstancedFoliageActor::GetInstancedFoliageActorForCurrentLevel(InComponent->GetWorld());
+	AInstancedFoliageActor* NewIFA = AInstancedFoliageActor::GetInstancedFoliageActorForCurrentLevel(InComponent->GetWorld(), true);
 	const auto SourceBaseId = InstanceBaseCache.GetInstanceBaseId(InComponent);
 	
 	for (auto& MeshPair : FoliageMeshes)
@@ -1171,12 +1192,7 @@ void AInstancedFoliageActor::MoveInstancesForComponentToCurrentLevel(UActorCompo
 
 void AInstancedFoliageActor::MoveInstancesToNewComponent(UPrimitiveComponent* InOldComponent, UPrimitiveComponent* InNewComponent)
 {
-	AInstancedFoliageActor* NewIFA = AInstancedFoliageActor::GetInstancedFoliageActorForLevel(InNewComponent->GetTypedOuter<ULevel>());
-	if (!NewIFA)
-	{
-		NewIFA = InNewComponent->GetWorld()->SpawnActor<AInstancedFoliageActor>();
-	}
-	check(NewIFA);
+	AInstancedFoliageActor* NewIFA = AInstancedFoliageActor::GetInstancedFoliageActorForLevel(InNewComponent->GetTypedOuter<ULevel>(), true);
 
 	const auto OldBaseId = this->InstanceBaseCache.GetInstanceBaseId(InOldComponent);
 	const auto NewBaseId = NewIFA->InstanceBaseCache.AddInstanceBaseId(InNewComponent);
@@ -1314,6 +1330,45 @@ FFoliageMeshInfo* AInstancedFoliageActor::UpdateMeshSettings(const UStaticMesh*
 }
 #endif
 
+UFoliageType* AInstancedFoliageActor::AddFoliageType(const UFoliageType* InType, FFoliageMeshInfo** OutInfo)
+{
+	FFoliageMeshInfo* MeshInfo = nullptr;
+	UFoliageType* FoliageType = const_cast<UFoliageType*>(InType);
+	
+	if (FoliageType->GetOuter() == this || FoliageType->IsAsset())
+	{
+		auto* MeshInfoPtr = FoliageMeshes.Find(FoliageType);
+		if (!MeshInfoPtr)
+		{
+			MarkPackageDirty();
+			MeshInfo = &FoliageMeshes.Add(FoliageType).Get();
+		}
+		else
+		{
+			MeshInfo = &MeshInfoPtr->Get();
+		}
+	}
+	else
+	{
+		// Unique meshes only
+		// Multiple entries for same static mesh can be added using FoliageType as an asset
+		FoliageType = GetSettingsForMesh(FoliageType->GetStaticMesh());
+		if (FoliageType == nullptr)
+		{
+			FoliageType = DuplicateObject<UFoliageType>(FoliageType, this);
+			MarkPackageDirty();
+			MeshInfo = &FoliageMeshes.Add(FoliageType).Get();
+		}
+	}
+
+	if (OutInfo)
+	{
+		*OutInfo = MeshInfo;
+	}
+
+	return FoliageType;
+}
+
 FFoliageMeshInfo* AInstancedFoliageActor::AddMesh(UStaticMesh* InMesh, UFoliageType** OutSettings, const UFoliageType_InstancedStaticMesh* DefaultSettings)
 {
 	check(GetSettingsForMesh(InMesh) == nullptr);
@@ -1371,22 +1426,28 @@ FFoliageMeshInfo* AInstancedFoliageActor::AddMesh(UFoliageType* InType)
 	return MeshInfo;
 }
 
-void AInstancedFoliageActor::RemoveMesh(UFoliageType* InSettings)
+void AInstancedFoliageActor::RemoveFoliageType(UFoliageType** InFoliageTypes, int32 Num)
 {
 	Modify();
-	MarkPackageDirty();
 	UnregisterAllComponents();
 
 	// Remove all components for this mesh from the Components array.
-	FFoliageMeshInfo* MeshInfo = FindMesh(InSettings);
-	if (MeshInfo && MeshInfo->Component)
+	for (int32 FoliageTypeIdx = 0; FoliageTypeIdx < Num; ++FoliageTypeIdx)
 	{
-		MeshInfo->Component->bAutoRegister = false;
-	}
+		const UFoliageType* FoliageType = InFoliageTypes[FoliageTypeIdx];
+		FFoliageMeshInfo* MeshInfo = FindMesh(FoliageType);
+		if (MeshInfo)
+		{
+			if (MeshInfo->Component)
+			{
+				MeshInfo->Component->bAutoRegister = false;
+			}
 
-	FoliageMeshes.Remove(InSettings);
+			FoliageMeshes.Remove(FoliageType);
+		}
+	}
+		
 	RegisterAllComponents();
-	CheckSelection();
 }
 
 void AInstancedFoliageActor::SelectInstance(UInstancedStaticMeshComponent* InComponent, int32 InInstanceIndex, bool bToggle)
@@ -1420,54 +1481,47 @@ void AInstancedFoliageActor::SelectInstance(UInstancedStaticMeshComponent* InCom
 		UFoliageType* Type = nullptr;
 		FFoliageMeshInfo* MeshInfo = nullptr;
 
-		Type = GetSettingsForMesh(InComponent->StaticMesh, &MeshInfo);
-
-		if (MeshInfo)
+		for (auto& MeshPair : FoliageMeshes)
 		{
-			if (InComponent == MeshInfo->Component)
+			if (MeshPair.Value->Component == InComponent)
 			{
-				bool bIsSelected = MeshInfo->SelectedIndices.Contains(InInstanceIndex);
+				Type = MeshPair.Key;
+				MeshInfo = &MeshPair.Value.Get();
+			}
+		}
+		
+		if (MeshInfo)
+		{
+			bool bIsSelected = MeshInfo->SelectedIndices.Contains(InInstanceIndex);
 
-				// Deselect if it's already selected.
-				if (InInstanceIndex < InComponent->SelectedInstances.Num())
-				{
-					InComponent->SelectedInstances[InInstanceIndex] = false;
-					InComponent->ReleasePerInstanceRenderData();
-					InComponent->MarkRenderStateDirty();
-				}
+			// Deselect if it's already selected.
+			if (InInstanceIndex < InComponent->SelectedInstances.Num())
+			{
+				InComponent->SelectedInstances[InInstanceIndex] = false;
+				InComponent->ReleasePerInstanceRenderData();
+				InComponent->MarkRenderStateDirty();
+			}
 
-				if (bIsSelected)
-				{
-					MeshInfo->SelectedIndices.Remove(InInstanceIndex);
-				}
+			if (bIsSelected)
+			{
+				MeshInfo->SelectedIndices.Remove(InInstanceIndex);
+			}
 
-				if (bToggle && bIsSelected)
+			if (!bToggle || !bIsSelected)
+			{
+				// Add the selection
+				if (InComponent->SelectedInstances.Num() < InComponent->PerInstanceSMData.Num())
 				{
-					if (SelectedMesh == Type && MeshInfo->SelectedIndices.Num() == 0)
-					{
-						SelectedMesh = nullptr;
-					}
+					InComponent->SelectedInstances.Init(false, InComponent->PerInstanceSMData.Num());
 				}
-				else
-				{
-					// Add the selection
-					if (InComponent->SelectedInstances.Num() < InComponent->PerInstanceSMData.Num())
-					{
-						InComponent->SelectedInstances.Init(false, InComponent->PerInstanceSMData.Num());
-					}
-					InComponent->SelectedInstances[InInstanceIndex] = true;
-					InComponent->ReleasePerInstanceRenderData();
-					InComponent->MarkRenderStateDirty();
+				InComponent->SelectedInstances[InInstanceIndex] = true;
+				InComponent->ReleasePerInstanceRenderData();
+				InComponent->MarkRenderStateDirty();
 
-					SelectedMesh = Type;
-					MeshInfo->SelectedIndices.Add(InInstanceIndex);
-				}
+				MeshInfo->SelectedIndices.Add(InInstanceIndex);
 			}
 		}
 	}
-
-	CheckSelection();
-
 }
 
 void AInstancedFoliageActor::PostEditUndo()
@@ -1484,8 +1538,6 @@ void AInstancedFoliageActor::PostEditUndo()
 
 void AInstancedFoliageActor::ApplySelectionToComponents(bool bApply)
 {
-	bool bNeedsUpdate = false;
-
 	for (auto& MeshPair : FoliageMeshes)
 	{
 		FFoliageMeshInfo& MeshInfo = *MeshPair.Value;
@@ -1498,12 +1550,13 @@ void AInstancedFoliageActor::ApplySelectionToComponents(bool bApply)
 
 				// Apply any selections in the component
 				MeshInfo.Component->SelectedInstances.Init(false, MeshInfo.Component->PerInstanceSMData.Num());
-				MeshInfo.Component->ReleasePerInstanceRenderData();
-				MeshInfo.Component->MarkRenderStateDirty();
 				for (int32 i : MeshInfo.SelectedIndices)
 				{
 					MeshInfo.Component->SelectedInstances[i] = true;
 				}
+
+				MeshInfo.Component->ReleasePerInstanceRenderData();
+				MeshInfo.Component->MarkRenderStateDirty();
 			}
 		}
 		else		
@@ -1512,6 +1565,7 @@ void AInstancedFoliageActor::ApplySelectionToComponents(bool bApply)
 			{
 				// remove any selections in the component
 				MeshInfo.Component->SelectedInstances.Empty();
+
 				MeshInfo.Component->ReleasePerInstanceRenderData();
 				MeshInfo.Component->MarkRenderStateDirty();
 			}
@@ -1519,46 +1573,19 @@ void AInstancedFoliageActor::ApplySelectionToComponents(bool bApply)
 	}
 }
 
-void AInstancedFoliageActor::CheckSelection()
+bool AInstancedFoliageActor::GetSelectionLocation(FVector& OutLocation) const
 {
-	// Check if we have to change the selection.
-	if (SelectedMesh != nullptr)
+	for (const auto& MeshPair: FoliageMeshes)
 	{
-		FFoliageMeshInfo* MeshInfo = FindMesh(SelectedMesh);
-		if (MeshInfo && MeshInfo->SelectedIndices.Num() > 0)
+		const FFoliageMeshInfo& MeshInfo = MeshPair.Value.Get();
+		if (MeshInfo.SelectedIndices.Num())
 		{
-			return;
-		}
-	}
-
-	SelectedMesh = nullptr;
-
-	// Try to find a new selection
-	for (auto& MeshPair : FoliageMeshes)
-	{
-		FFoliageMeshInfo& MeshInfo = *MeshPair.Value;
-		if (MeshInfo.SelectedIndices.Num() > 0)
-		{
-			SelectedMesh = MeshPair.Key;
-			return;
-		}
-	}
-}
-
-FVector AInstancedFoliageActor::GetSelectionLocation()
-{
-	FVector Result(0, 0, 0);
-
-	if (SelectedMesh != nullptr)
-	{
-		FFoliageMeshInfo* MeshInfo = FindMesh(SelectedMesh);
-		if (MeshInfo && MeshInfo->SelectedIndices.Num() > 0)
-		{
-			Result = MeshInfo->Instances[MeshInfo->SelectedIndices.Array()[0]].Location;
+			const int32 InstanceIdx = (*MeshInfo.SelectedIndices.CreateConstIterator());
+			OutLocation = MeshInfo.Instances[InstanceIdx].Location;
+			return true;
 		}
 	}
-
-	return Result;
+	return false;
 }
 
 void AInstancedFoliageActor::MapRebuild()
@@ -1762,6 +1789,8 @@ void AInstancedFoliageActor::PostLoad()
 {
 	Super::PostLoad();
 
+	GetLevel()->InstancedFoliageActor = this;
+
 #if WITH_EDITOR
 	if (GIsEditor)
 	{
@@ -1853,6 +1882,13 @@ void AInstancedFoliageActor::PostLoad()
 #endif
 }
 
+void AInstancedFoliageActor::PostRegisterAllComponents()
+{
+	Super::PostRegisterAllComponents();
+	//
+	GetLevel()->InstancedFoliageActor = this;
+}
+
 #if WITH_EDITOR
 
 void AInstancedFoliageActor::NotifyFoliageTypeChanged(UFoliageType* FoliageType)
@@ -1904,10 +1940,6 @@ void AInstancedFoliageActor::OnLevelActorDeleted(AActor* InActor)
 void AInstancedFoliageActor::AddReferencedObjects(UObject* InThis, FReferenceCollector& Collector)
 {
 	AInstancedFoliageActor* This = CastChecked<AInstancedFoliageActor>(InThis);
-	if (This->SelectedMesh)
-	{
-		Collector.AddReferencedObject(This->SelectedMesh, This);
-	}
 
 	for (auto& MeshPair : This->FoliageMeshes)
 	{
@@ -1919,6 +1951,7 @@ void AInstancedFoliageActor::AddReferencedObjects(UObject* InThis, FReferenceCol
 			Collector.AddReferencedObject(MeshInfo.Component, This);
 		}
 	}
+	
 	Super::AddReferencedObjects(This, Collector);
 }
 
@@ -1948,9 +1981,9 @@ void AInstancedFoliageActor::ApplyWorldOffset(const FVector& InOffset, bool bWor
 	}
 }
 
-bool AInstancedFoliageActor::FoliageTrace(const UWorld* InWorld, FHitResult& OutHit, const FDesiredFoliageInstance& DesiredInstance, const AInstancedFoliageActor* IgnoreIFA, FName InTraceTag, bool InbReturnFaceIndex)
+bool AInstancedFoliageActor::FoliageTrace(const UWorld* InWorld, FHitResult& OutHit, const FDesiredFoliageInstance& DesiredInstance, FName InTraceTag, bool InbReturnFaceIndex)
 {
-	FCollisionQueryParams QueryParams(InTraceTag, true, IgnoreIFA);
+	FCollisionQueryParams QueryParams(InTraceTag, true);
 	QueryParams.bReturnFaceIndex = InbReturnFaceIndex;
 
 	FVector StartTrace = DesiredInstance.StartTrace;
@@ -2006,21 +2039,20 @@ bool AInstancedFoliageActor::FoliageTrace(const UWorld* InWorld, FHitResult& Out
 	return false;
 }
 
-bool AInstancedFoliageActor::CheckCollisionWithWorld(const UFoliageType* Settings, const FFoliageInstance& Inst, const FVector& HitNormal, const FVector& HitLocation)
+bool AInstancedFoliageActor::CheckCollisionWithWorld(const UWorld* InWorld, const UFoliageType* Settings, const FFoliageInstance& Inst, const FVector& HitNormal, const FVector& HitLocation)
 {
 	FMatrix InstTransform = Inst.GetInstanceWorldTransform().ToMatrixWithScale();
 	FVector LocalHit = InstTransform.InverseTransformPosition(HitLocation);
 
-	UWorld* World = GetWorld();
-
 	if (Settings->CollisionWithWorld)
 	{
 		// Check for overhanging ledge
 		{
-			FVector LocalSamplePos[4] = { FVector(Settings->LowBoundOriginRadius.Z, 0, 0),
-			FVector(-Settings->LowBoundOriginRadius.Z, 0, 0),
-			FVector(0, Settings->LowBoundOriginRadius.Z, 0),
-			FVector(0, -Settings->LowBoundOriginRadius.Z, 0)
+			FVector LocalSamplePos[4] = { 
+				FVector(Settings->LowBoundOriginRadius.Z, 0, 0),
+				FVector(-Settings->LowBoundOriginRadius.Z, 0, 0),
+				FVector(0, Settings->LowBoundOriginRadius.Z, 0),
+				FVector(0, -Settings->LowBoundOriginRadius.Z, 0)
 			};
 
 
@@ -2030,7 +2062,7 @@ bool AInstancedFoliageActor::CheckCollisionWithWorld(const UFoliageType* Setting
 				FVector SamplePos = InstTransform.TransformPosition(FVector(Settings->LowBoundOriginRadius.X, Settings->LowBoundOriginRadius.Y, 2.f) + LocalSamplePos[i]);
 				float WorldRadius = (Settings->LowBoundOriginRadius.Z + 2.f)*FMath::Max(Inst.DrawScale3D.X, Inst.DrawScale3D.Y);
 				FVector NormalVector = Settings->AlignToNormal ? HitNormal : FVector(0, 0, 1);
-				if (AInstancedFoliageActor::FoliageTrace(World, Hit, FDesiredFoliageInstance(SamplePos, SamplePos - NormalVector*WorldRadius), this))
+				if (AInstancedFoliageActor::FoliageTrace(InWorld, Hit, FDesiredFoliageInstance(SamplePos, SamplePos - NormalVector*WorldRadius)))
 				{
 					if (LocalHit.Z - Inst.ZOffset < Settings->LowBoundOriginRadius.Z)
 					{
@@ -2043,33 +2075,31 @@ bool AInstancedFoliageActor::CheckCollisionWithWorld(const UFoliageType* Setting
 
 		// Check collision with Bounding Box
 		{
-		FBox MeshBox = Settings->MeshBounds.GetBox();
-		MeshBox.Min.Z = FMath::Min(MeshBox.Max.Z, LocalHit.Z + Settings->MeshBounds.BoxExtent.Z * 0.05f);
-		FBoxSphereBounds ShrinkBound(MeshBox);
-		FBoxSphereBounds WorldBound = ShrinkBound.TransformBy(InstTransform);
-		//::DrawDebugBox(World, WorldBound.Origin, WorldBound.BoxExtent, FColor::Red, true, 10.f);
-		static FName NAME_FoliageCollisionWithWorld = FName(TEXT("FoliageCollisionWithWorld"));
-		if (World->OverlapTest(WorldBound.Origin, FQuat(Inst.Rotation), FCollisionShape::MakeBox(ShrinkBound.BoxExtent * Inst.DrawScale3D * Settings->CollisionScale), FCollisionQueryParams(NAME_FoliageCollisionWithWorld, false), FCollisionObjectQueryParams(ECC_WorldStatic)))
-		{
-			return false;
+			FBox MeshBox = Settings->MeshBounds.GetBox();
+			MeshBox.Min.Z = FMath::Min(MeshBox.Max.Z, LocalHit.Z + Settings->MeshBounds.BoxExtent.Z * 0.05f);
+			FBoxSphereBounds ShrinkBound(MeshBox);
+			FBoxSphereBounds WorldBound = ShrinkBound.TransformBy(InstTransform);
+			//::DrawDebugBox(World, WorldBound.Origin, WorldBound.BoxExtent, FColor::Red, true, 10.f);
+			static FName NAME_FoliageCollisionWithWorld = FName(TEXT("FoliageCollisionWithWorld"));
+			if (InWorld->OverlapTest(WorldBound.Origin, FQuat(Inst.Rotation), FCollisionShape::MakeBox(ShrinkBound.BoxExtent * Inst.DrawScale3D * Settings->CollisionScale), FCollisionQueryParams(NAME_FoliageCollisionWithWorld, false), FCollisionObjectQueryParams(ECC_WorldStatic)))
+			{
+				return false;
+			}
 		}
 	}
-	}
 
 	return true;
 }
 
-FPotentialInstance::FPotentialInstance(FVector InHitLocation, FVector InHitNormal, UPrimitiveComponent* InHitComponent, float InHitWeight, FFoliageMeshInfo* InMeshInfo, AInstancedFoliageActor* InIFA, const FDesiredFoliageInstance& InDesiredInstance)
+FPotentialInstance::FPotentialInstance(FVector InHitLocation, FVector InHitNormal, UPrimitiveComponent* InHitComponent, float InHitWeight, const FDesiredFoliageInstance& InDesiredInstance)
 : HitLocation(InHitLocation)
 , HitNormal(InHitNormal)
 , HitComponent(InHitComponent)
 , HitWeight(InHitWeight)
-, MeshInfo(InMeshInfo)
-, IFA(InIFA)
 , DesiredInstance(InDesiredInstance)
 {}
 
-bool FPotentialInstance::PlaceInstance(const UFoliageType* Settings, FFoliageInstance& Inst, bool bSkipCollision)
+bool FPotentialInstance::PlaceInstance(const UWorld* InWorld, const UFoliageType* Settings, FFoliageInstance& Inst, bool bSkipCollision)
 {
 	if (DesiredInstance.PlacementMode != EFoliagePlacementMode::Procedural)
 	{
@@ -2139,11 +2169,7 @@ bool FPotentialInstance::PlaceInstance(const UFoliageType* Settings, FFoliageIns
 		}
 	}
 
-#if WITH_EDITORONLY_DATA
-	Inst.BaseId = IFA->InstanceBaseCache.AddInstanceBaseId(HitComponent);
-#endif//WITH_EDITORONLY_DATA
-
-	return bSkipCollision || IFA->CheckCollisionWithWorld(Settings, Inst, HitNormal, HitLocation);
+	return bSkipCollision || AInstancedFoliageActor::CheckCollisionWithWorld(InWorld, Settings, Inst, HitNormal, HitLocation);
 }
 
 #undef LOCTEXT_NAMESPACE
diff --git a/Engine/Source/Runtime/Foliage/Private/ProceduralFoliageComponent.cpp b/Engine/Source/Runtime/Foliage/Private/ProceduralFoliageComponent.cpp
--- a/Engine/Source/Runtime/Foliage/Private/ProceduralFoliageComponent.cpp
+++ b/Engine/Source/Runtime/Foliage/Private/ProceduralFoliageComponent.cpp
@@ -259,7 +259,10 @@ void UProceduralFoliageComponent::RemoveProceduralContent()
 		if (Level)
 		{
 			AInstancedFoliageActor* IFA = AInstancedFoliageActor::GetInstancedFoliageActorForLevel(Level);
-			IFA->DeleteInstancesForProceduralFoliageComponent(this);
+			if (IFA)
+			{
+				IFA->DeleteInstancesForProceduralFoliageComponent(this);
+			}
 		}
 	}
 #endif
diff --git a/Engine/Source/Runtime/Foliage/Public/FoliageType.h b/Engine/Source/Runtime/Foliage/Public/FoliageType.h
--- a/Engine/Source/Runtime/Foliage/Public/FoliageType.h
+++ b/Engine/Source/Runtime/Foliage/Public/FoliageType.h
@@ -237,8 +237,9 @@ class UFoliageType : public UObject
 	UPROPERTY(EditAnywhere, Category = Painting, meta = (Subcategory = "Collision", HideObjectType = true))
 	struct FBodyInstance BodyInstance;
 
-	/* Gets the mesh associated with this FoliageType */
+	/* Gets/Sets the mesh associated with this FoliageType */
 	virtual UStaticMesh* GetStaticMesh() const PURE_VIRTUAL(UFoliageType::GetStaticMesh, return nullptr; );
+	virtual void SetStaticMesh(UStaticMesh* InStaticMesh) PURE_VIRTUAL(UFoliageType::SetStaticMesh,);
 
 #if WITH_EDITOR
 	/* Lets subclasses decide if the InstancedFoliageActor should reallocate its instances if the specified property change event occurs */
diff --git a/Engine/Source/Runtime/Foliage/Public/FoliageType_InstancedStaticMesh.h b/Engine/Source/Runtime/Foliage/Public/FoliageType_InstancedStaticMesh.h
--- a/Engine/Source/Runtime/Foliage/Public/FoliageType_InstancedStaticMesh.h
+++ b/Engine/Source/Runtime/Foliage/Public/FoliageType_InstancedStaticMesh.h
@@ -17,4 +17,9 @@ class UFoliageType_InstancedStaticMesh : public UFoliageType
 	{
 		return Mesh;
 	}
+
+	virtual void SetStaticMesh(UStaticMesh* InStaticMesh) override
+	{
+		Mesh = InStaticMesh;
+	}
 };
diff --git a/Engine/Source/Runtime/Foliage/Public/InstancedFoliage.h b/Engine/Source/Runtime/Foliage/Public/InstancedFoliage.h
--- a/Engine/Source/Runtime/Foliage/Public/InstancedFoliage.h
+++ b/Engine/Source/Runtime/Foliage/Public/InstancedFoliage.h
@@ -202,6 +202,8 @@ struct FFoliageMeshInfo
 	FOLIAGE_API void ReapplyInstancesToComponent();
 
 	FOLIAGE_API void SelectInstances(AInstancedFoliageActor* InIFA, bool bSelect, TArray<int32>& Instances);
+	
+	FOLIAGE_API void SelectInstances(AInstancedFoliageActor* InIFA, bool bSelect);
 
 	// Get the number of placed instances
 	FOLIAGE_API int32 GetInstanceCount() const;
@@ -347,7 +349,6 @@ struct FDesiredFoliageInstance
 {
 	FDesiredFoliageInstance()
 	: FoliageType(nullptr)
-	, MeshInfo(nullptr)
 	, StartTrace(ForceInit)
 	, EndTrace(ForceInit)
 	, Rotation(ForceInit)
@@ -360,7 +361,6 @@ struct FDesiredFoliageInstance
 
 	FDesiredFoliageInstance(const FVector& InStartTrace, const FVector& InEndTrace, const float InTraceRadius = 0.f)
 	: FoliageType(nullptr)
-	, MeshInfo(nullptr)
 	, StartTrace(InStartTrace)
 	, EndTrace(InEndTrace)
 	, Rotation(ForceInit)
@@ -370,9 +370,8 @@ struct FDesiredFoliageInstance
 	{
 	}
 
-	FGuid ProceduralGuid;
 	const UFoliageType* FoliageType;
-	FFoliageMeshInfo* MeshInfo;
+	FGuid ProceduralGuid;
 	FVector StartTrace;
 	FVector EndTrace;
 	FQuat Rotation;
@@ -389,10 +388,8 @@ struct FOLIAGE_API FPotentialInstance
 	FVector HitNormal;
 	UPrimitiveComponent* HitComponent;
 	float HitWeight;
-	FFoliageMeshInfo* MeshInfo;
-	AInstancedFoliageActor* IFA;
 	FDesiredFoliageInstance DesiredInstance;
 
-	FPotentialInstance(FVector InHitLocation, FVector InHitNormal, UPrimitiveComponent* InHitComponent, float InHitWeight, FFoliageMeshInfo* MeshInfo = nullptr, AInstancedFoliageActor* InIFA = nullptr, const FDesiredFoliageInstance& InDesiredInstance = FDesiredFoliageInstance());
-	bool PlaceInstance(const UFoliageType* Settings, FFoliageInstance& Inst, bool bSkipCollision = false);
+	FPotentialInstance(FVector InHitLocation, FVector InHitNormal, UPrimitiveComponent* InHitComponent, float InHitWeight, const FDesiredFoliageInstance& InDesiredInstance = FDesiredFoliageInstance());
+	bool PlaceInstance(const UWorld* InWorld, const UFoliageType* Settings, FFoliageInstance& Inst, bool bSkipCollision = false);
 };
diff --git a/Engine/Source/Runtime/Foliage/Public/InstancedFoliageActor.h b/Engine/Source/Runtime/Foliage/Public/InstancedFoliageActor.h
--- a/Engine/Source/Runtime/Foliage/Public/InstancedFoliageActor.h
+++ b/Engine/Source/Runtime/Foliage/Public/InstancedFoliageActor.h
@@ -25,10 +25,6 @@ class AInstancedFoliageActor : public AActor
 {
 	GENERATED_UCLASS_BODY()
 
-	/** The static mesh type that will be used to show the widget */
-	UPROPERTY(transient)
-	UFoliageType* SelectedMesh;
-
 public:
 #if WITH_EDITORONLY_DATA
 	// Cross level references cache for instances base
@@ -43,6 +39,7 @@ public:
 	virtual void PostInitProperties() override;
 	virtual void BeginDestroy() override;
 	virtual void PostLoad() override;
+	virtual void PostRegisterAllComponents() override;
 	static void AddReferencedObjects(UObject* InThis, FReferenceCollector& Collector);
 	// End UObject interface. 
 
@@ -76,7 +73,7 @@ public:
 	* @param bCreateIfNone					Create if doesnt already exist
 	* returns								pointer to foliage object instance
 	*/
-	static FOLIAGE_API AInstancedFoliageActor* GetInstancedFoliageActorForCurrentLevel(UWorld* InWorld, bool bCreateIfNone = true);
+	static FOLIAGE_API AInstancedFoliageActor* GetInstancedFoliageActorForCurrentLevel(UWorld* InWorld, bool bCreateIfNone = false);
 
 
 	/**
@@ -84,10 +81,10 @@ public:
 	* @param bCreateIfNone					Create if doesnt already exist
 	* returns								pointer to foliage object instance
 	*/
-	static FOLIAGE_API AInstancedFoliageActor* GetInstancedFoliageActorForLevel(ULevel* Level, bool bCreateIfNone = true);
+	static FOLIAGE_API AInstancedFoliageActor* GetInstancedFoliageActorForLevel(ULevel* Level, bool bCreateIfNone = false);
 
-	static FOLIAGE_API bool FoliageTrace(const UWorld* InWorld, FHitResult& OutHit, const FDesiredFoliageInstance& DesiredInstance, const AInstancedFoliageActor* IgnoreIFA = nullptr, FName InTraceTag = NAME_None, bool InbReturnFaceIndex = false);
-	FOLIAGE_API bool CheckCollisionWithWorld(const UFoliageType* Settings, const FFoliageInstance& Inst, const FVector& HitNormal, const FVector& HitLocation);
+	static FOLIAGE_API bool FoliageTrace(const UWorld* InWorld, FHitResult& OutHit, const FDesiredFoliageInstance& DesiredInstance, FName InTraceTag = NAME_None, bool InbReturnFaceIndex = false);
+	static FOLIAGE_API bool CheckCollisionWithWorld(const UWorld* InWorld, const UFoliageType* Settings, const FFoliageInstance& Inst, const FVector& HitNormal, const FVector& HitLocation);
 
 #if WITH_EDITOR
 	virtual void PostEditUndo() override;
@@ -110,6 +107,7 @@ public:
 
 	// Deletes the instances attached to a component
 	FOLIAGE_API void DeleteInstancesForComponent(UActorComponent* InComponent);
+	FOLIAGE_API void DeleteInstancesForComponent(UActorComponent* InComponent, const UFoliageType* InFoliageType);
 
 	// Deletes the instances spawned by a procedural component
 	void DeleteInstancesForProceduralFoliageComponent(const UProceduralFoliageComponent* ProceduralComponent);
@@ -117,13 +115,14 @@ public:
 	// Finds a mesh entry or adds it if it doesn't already exist
 	FOLIAGE_API FFoliageMeshInfo* FindOrAddMesh(UFoliageType* InType);
 
+	FOLIAGE_API UFoliageType* AddFoliageType(const UFoliageType* InType, FFoliageMeshInfo** OutInfo = nullptr);
 	// Add a new static mesh.
 	FOLIAGE_API FFoliageMeshInfo* AddMesh(UStaticMesh* InMesh, UFoliageType** OutSettings = nullptr, const UFoliageType_InstancedStaticMesh* DefaultSettings = nullptr);
 	FOLIAGE_API FFoliageMeshInfo* AddMesh(UFoliageType* InType);
 	FOLIAGE_API FFoliageMeshInfo* UpdateMeshSettings(const UStaticMesh* InMesh, const UFoliageType_InstancedStaticMesh* DefaultSettings);
 
-	// Remove the static mesh from the mesh list, and all its instances.
-	FOLIAGE_API void RemoveMesh(UFoliageType* InFoliageType);
+	// Remove the FoliageType from the list, and all its instances.
+	FOLIAGE_API void RemoveFoliageType(UFoliageType** InFoliageType, int32 Num);
 
 	// Select an individual instance.
 	FOLIAGE_API void SelectInstance(UInstancedStaticMeshComponent* InComponent, int32 InComponentInstanceIndex, bool bToggle);
@@ -131,11 +130,8 @@ public:
 	// Propagate the selected instances to the actual render components
 	FOLIAGE_API void ApplySelectionToComponents(bool bApply);
 
-	// Updates the SelectedMesh property of the actor based on the actual selected instances
-	FOLIAGE_API void CheckSelection();
-
 	// Returns the location for the widget
-	FOLIAGE_API FVector GetSelectionLocation();
+	FOLIAGE_API bool GetSelectionLocation(FVector& OutLocation) const;
 
 	/* Called to notify InstancedFoliageActor that a UFoliageType has been modified */
 	void NotifyFoliageTypeChanged(UFoliageType* FoliageType);
diff --git a/Engine/Source/Runtime/Landscape/Private/LandscapeEdit.cpp b/Engine/Source/Runtime/Landscape/Private/LandscapeEdit.cpp
--- a/Engine/Source/Runtime/Landscape/Private/LandscapeEdit.cpp
+++ b/Engine/Source/Runtime/Landscape/Private/LandscapeEdit.cpp
@@ -756,7 +756,10 @@ void ULandscapeComponent::UpdateCollisionHeightData(const FColor* HeightmapTextu
 		if (OldCollisionComponent && OldCollisionComponent != CollisionComp)
 		{
 			AInstancedFoliageActor* IFA = AInstancedFoliageActor::GetInstancedFoliageActorForCurrentLevel(OldCollisionComponent->GetWorld());
-			IFA->MoveInstancesToNewComponent(OldCollisionComponent, CollisionComp);
+			if (IFA)
+			{
+				IFA->MoveInstancesToNewComponent(OldCollisionComponent, CollisionComp);
+			}
 		}
 	}
 
