WIP: Foliage UI, can be activated by 'foliage.NewUI 1'

diff --git a/Engine/Source/Editor/FoliageEdit/Private/FoliageEdMode.cpp b/Engine/Source/Editor/FoliageEdit/Private/FoliageEdMode.cpp
--- a/Engine/Source/Editor/FoliageEdit/Private/FoliageEdMode.cpp
+++ b/Engine/Source/Editor/FoliageEdit/Private/FoliageEdMode.cpp
@@ -420,7 +420,7 @@ void FEdModeFoliage::GetRandomVectorInBrush(FVector& OutStart, FVector& OutEnd)
 static bool CheckLocationForPotentialInstance_ThreadSafe(const UFoliageType* Settings, const FVector& Location, const FVector& Normal)
 {
 	// Check height range
-	if (Location.Z < Settings->HeightMin || Location.Z > Settings->HeightMax)
+	if (!Settings->Height.Contains(Location.Z))
 	{
 		return false;
 	}
@@ -1280,7 +1280,7 @@ void FEdModeFoliage::ReapplyInstancesForBrush(UWorld* InWorld, AInstancedFoliage
 			{
 				if (Settings->ReapplyScaleX)
 				{
-					float Scale = Settings->ScaleMinX + FMath::FRand() * (Settings->ScaleMaxX - Settings->ScaleMinX);
+					float Scale = Settings->ScaleX.Interpolate(FMath::FRand());
 					Instance.DrawScale3D = FVector(Scale, Scale, Scale);
 					bUpdated = true;
 				}
@@ -1289,17 +1289,17 @@ void FEdModeFoliage::ReapplyInstancesForBrush(UWorld* InWorld, AInstancedFoliage
 			{
 				float LockRand;
 				// If we're doing axis scale locking, get an existing scale for a locked axis that we're not changing, for use as the locked scale value.
-				if (Settings->LockScaleX && !Settings->ReapplyScaleX && (Settings->ScaleMaxX - Settings->ScaleMinX) > KINDA_SMALL_NUMBER)
+				if (Settings->LockScaleX && !Settings->ReapplyScaleX && Settings->ScaleX.Size() > KINDA_SMALL_NUMBER)
 				{
-					LockRand = (Instance.DrawScale3D.X - Settings->ScaleMinX) / (Settings->ScaleMaxX - Settings->ScaleMinX);
+					LockRand = (Instance.DrawScale3D.X - Settings->ScaleX.Min) / Settings->ScaleX.Size();
 				}
-				else if (Settings->LockScaleY && !Settings->ReapplyScaleY && (Settings->ScaleMaxY - Settings->ScaleMinY) > KINDA_SMALL_NUMBER)
+				else if (Settings->LockScaleY && !Settings->ReapplyScaleY && Settings->ScaleY.Size() > KINDA_SMALL_NUMBER)
 				{
-					LockRand = (Instance.DrawScale3D.Y - Settings->ScaleMinY) / (Settings->ScaleMaxY - Settings->ScaleMinY);
+					LockRand = (Instance.DrawScale3D.Y - Settings->ScaleY.Min) / Settings->ScaleY.Size();
 				}
-				else if (Settings->LockScaleZ && !Settings->ReapplyScaleZ && (Settings->ScaleMaxZ - Settings->ScaleMinZ) > KINDA_SMALL_NUMBER)
+				else if (Settings->LockScaleZ && !Settings->ReapplyScaleZ && Settings->ScaleZ.Size() > KINDA_SMALL_NUMBER)
 				{
-					LockRand = (Instance.DrawScale3D.Z - Settings->ScaleMinZ) / (Settings->ScaleMaxZ - Settings->ScaleMinZ);
+					LockRand = (Instance.DrawScale3D.Z - Settings->ScaleZ.Min) / Settings->ScaleZ.Size();
 				}
 				else
 				{
@@ -1308,26 +1308,26 @@ void FEdModeFoliage::ReapplyInstancesForBrush(UWorld* InWorld, AInstancedFoliage
 
 				if (Settings->ReapplyScaleX)
 				{
-					Instance.DrawScale3D.X = Settings->ScaleMinX + (Settings->LockScaleX ? LockRand : FMath::FRand()) * (Settings->ScaleMaxX - Settings->ScaleMinX);
+					Instance.DrawScale3D.X = Settings->ScaleX.Interpolate(Settings->LockScaleX ? LockRand : FMath::FRand());
 					bUpdated = true;
 				}
 
 				if (Settings->ReapplyScaleY)
 				{
-					Instance.DrawScale3D.Y = Settings->ScaleMinY + (Settings->LockScaleY ? LockRand : FMath::FRand()) * (Settings->ScaleMaxY - Settings->ScaleMinY);
+					Instance.DrawScale3D.Y = Settings->ScaleY.Interpolate(Settings->LockScaleY ? LockRand : FMath::FRand());
 					bUpdated = true;
 				}
 
 				if (Settings->ReapplyScaleZ)
 				{
-					Instance.DrawScale3D.Z = Settings->ScaleMinZ + (Settings->LockScaleZ ? LockRand : FMath::FRand()) * (Settings->ScaleMaxZ - Settings->ScaleMinZ);
+					Instance.DrawScale3D.Z = Settings->ScaleZ.Interpolate(Settings->LockScaleZ ? LockRand : FMath::FRand());
 					bUpdated = true;
 				}
 			}
 
 			if (Settings->ReapplyZOffset)
 			{
-				Instance.ZOffset = Settings->ZOffsetMin + FMath::FRand() * (Settings->ZOffsetMax - Settings->ZOffsetMin);
+				Instance.ZOffset = Settings->ZOffset.Interpolate(FMath::FRand());
 				bUpdated = true;
 			}
 
@@ -1416,7 +1416,7 @@ void FEdModeFoliage::ReapplyInstancesForBrush(UWorld* InWorld, AInstancedFoliage
 			// Cull instances that don't meet the height range
 			if (Settings->ReapplyHeight)
 			{
-				if (Instance.Location.Z < Settings->HeightMin || Instance.Location.Z > Settings->HeightMax)
+				if (!Settings->Height.Contains(Instance.Location.Z))
 				{
 					InstancesToDelete.Add(InstanceIndex);
 					if (bReapplyLocation)
diff --git a/Engine/Source/Editor/FoliageEdit/Private/FoliageEdModeToolkit.cpp b/Engine/Source/Editor/FoliageEdit/Private/FoliageEdModeToolkit.cpp
--- a/Engine/Source/Editor/FoliageEdit/Private/FoliageEdModeToolkit.cpp
+++ b/Engine/Source/Editor/FoliageEdit/Private/FoliageEdModeToolkit.cpp
@@ -3,9 +3,16 @@
 #include "UnrealEd.h"
 #include "FoliageEdModeToolkit.h"
 #include "SFoliageEdit.h"
+#include "SFoliageEdit2.h"
 
 #define LOCTEXT_NAMESPACE "FoliageEditMode"
 
+static TAutoConsoleVariable<int32> CVarFoliageNewUI(
+	TEXT("foliage.NewUI"),
+	0,
+	TEXT("Enables new foliage UI"));
+
+
 void FFoliageEdModeToolkit::RegisterTabSpawners(const TSharedRef<class FTabManager>& TabManager)
 {
 
@@ -18,7 +25,17 @@ void FFoliageEdModeToolkit::UnregisterTabSpawners(const TSharedRef<class FTabMan
 
 void FFoliageEdModeToolkit::Init(const TSharedPtr< class IToolkitHost >& InitToolkitHost)
 {
-	FoliageEdWidget = SNew(SFoliageEdit);
+	bool NewUI = CVarFoliageNewUI.GetValueOnAnyThread() != 0;
+	
+	if (NewUI)
+	{
+		FoliageEdWidget2 = SNew(SFoliageEdit2);
+	}
+	else
+	{
+		FoliageEdWidget = SNew(SFoliageEdit);
+	}
+	
 
 	FModeToolkit::Init(InitToolkitHost);
 }
@@ -40,12 +57,26 @@ class FEdMode* FFoliageEdModeToolkit::GetEditorMode() const
 
 TSharedPtr<SWidget> FFoliageEdModeToolkit::GetInlineContent() const
 {
-	return FoliageEdWidget;
+	if (FoliageEdWidget2.IsValid())
+	{
+		return FoliageEdWidget2;
+	}
+	else
+	{
+		return FoliageEdWidget;
+	}
 }
 
 void FFoliageEdModeToolkit::RefreshFullList()
 {
-	FoliageEdWidget->RefreshFullList();
+	if (FoliageEdWidget2.IsValid())
+	{
+		FoliageEdWidget2->RefreshFullList();
+	}
+	else
+	{
+		FoliageEdWidget->RefreshFullList();
+	}
 }
 
 #undef LOCTEXT_NAMESPACE
diff --git a/Engine/Source/Editor/FoliageEdit/Private/SFoliageEdit2.cpp b/Engine/Source/Editor/FoliageEdit/Private/SFoliageEdit2.cpp
new file mode 100644
--- /dev/null
+++ b/Engine/Source/Editor/FoliageEdit/Private/SFoliageEdit2.cpp
@@ -0,0 +1,884 @@
+// Copyright 1998-2015 Epic Games, Inc. All Rights Reserved.
+
+#include "UnrealEd.h"
+
+#include "SScaleBox.h"
+
+#include "SFoliageEdit2.h"
+#include "FoliageEditActions.h"
+#include "FoliageEdMode.h"
+#include "FoliageType.h"
+#include "Editor/UnrealEd/Public/AssetThumbnail.h"
+#include "Editor/UnrealEd/Public/DragAndDrop/AssetDragDropOp.h"
+#include "Editor/ContentBrowser/Public/ContentBrowserModule.h"
+#include "Editor/UnrealEd/Public/AssetSelection.h"
+#include "Editor/IntroTutorials/Public/IIntroTutorials.h"
+#include "Editor/PropertyEditor/Public/IDetailsView.h"
+#include "SFoliageEditMeshDisplayItem.h"
+#include "Engine/StaticMesh.h"
+#include "SFoliageEditMeshRow.h"
+
+#define LOCTEXT_NAMESPACE "FoliageEd_Mode"
+
+class SFoliageDragDropHandler2 : public SCompoundWidget
+{
+public:
+	SLATE_BEGIN_ARGS(SFoliageDragDropHandler2) {}
+	SLATE_DEFAULT_SLOT(FArguments, Content)
+	SLATE_ARGUMENT(TWeakPtr<SFoliageEdit2>, FoliageEditPtr)
+	SLATE_END_ARGS()
+
+	/** SCompoundWidget functions */
+	void Construct(const FArguments& InArgs)
+	{
+		FoliageEditPtr = InArgs._FoliageEditPtr;
+		bIsDragOn = false;
+
+		this->ChildSlot
+		[
+			SNew(SBorder)
+			.BorderImage(FEditorStyle::GetBrush("WhiteBrush"))
+			.BorderBackgroundColor(this, &SFoliageDragDropHandler2::GetBackgroundColor)
+			.Padding(FMargin(30))
+			[
+				InArgs._Content.Widget
+			]
+		];
+	}
+
+	~SFoliageDragDropHandler2() {}
+
+	FReply OnDrop(const FGeometry& MyGeometry, const FDragDropEvent& DragDropEvent) override
+	{
+		bIsDragOn = false;
+
+		TSharedPtr<SFoliageEdit2> PinnedPtr = FoliageEditPtr.Pin();
+		if (PinnedPtr.IsValid())
+		{
+			return PinnedPtr->OnDrop_ListView(MyGeometry, DragDropEvent);
+		}
+		return FReply::Handled();
+	}
+
+	virtual void OnDragEnter(const FGeometry& MyGeometry, const FDragDropEvent& DragDropEvent) override
+	{
+		bIsDragOn = true;
+	}
+
+	virtual void OnDragLeave(const FDragDropEvent& DragDropEvent) override
+	{
+		bIsDragOn = false;
+	}
+
+private:
+	FSlateColor GetBackgroundColor() const
+	{
+		return bIsDragOn ? FLinearColor(1.0f, 0.6f, 0.1f, 0.9f) : FLinearColor(0.1f, 0.1f, 0.1f, 0.9f);
+	}
+
+private:
+	TWeakPtr<SFoliageEdit2> FoliageEditPtr;
+	bool bIsDragOn;
+};
+
+
+BEGIN_SLATE_FUNCTION_BUILD_OPTIMIZATION
+void SFoliageEdit2::Construct(const FArguments& InArgs)
+{
+	FoliageEditMode = (FEdModeFoliage*)GLevelEditorModeTools().GetActiveMode(FBuiltinEditorModes::EM_Foliage);
+
+	FFoliageEditCommands::Register();
+
+	UICommandList = MakeShareable(new FUICommandList);
+
+	BindCommands();
+
+	AssetThumbnailPool = MakeShareable(new FAssetThumbnailPool(512, TAttribute<bool>::Create(TAttribute<bool>::FGetter::CreateSP(this, &SFoliageEdit2::IsHovered))));
+
+	FPropertyEditorModule& PropertyModule = FModuleManager::LoadModuleChecked<FPropertyEditorModule>("PropertyEditor");
+	FDetailsViewArgs Args(false, false, false, FDetailsViewArgs::HideNameArea, true);
+	Args.bShowActorLabel = false;
+	MeshDetailsWidget = PropertyModule.CreateDetailView(Args);
+
+	
+	// Everything (or almost) uses this padding, change it to expand the padding.
+	FMargin StandardPadding(0.0f, 4.0f, 0.0f, 4.0f);
+
+	this->ChildSlot
+	[
+		SNew(SOverlay)
+
+		+ SOverlay::Slot()
+		[
+			SNew(SVerticalBox)
+
+			+ SVerticalBox::Slot()
+			.AutoHeight()
+			.Padding(StandardPadding)
+			[
+				BuildToolBar()
+			]
+
+			+ SVerticalBox::Slot()
+			.AutoHeight()
+			[
+				SNew(SHorizontalBox)
+				.ToolTipText(LOCTEXT("BrushSize_Tooltip", "The size of the foliage brush"))
+
+				+ SHorizontalBox::Slot()
+				.FillWidth(1.0f)
+				.Padding(StandardPadding)
+				[
+					SNew(STextBlock)
+					.Visibility(this, &SFoliageEdit2::GetVisibility_Radius)
+					.Text(LOCTEXT("BrushSize", "Brush Size"))
+				]
+				+ SHorizontalBox::Slot()
+				.FillWidth(2.0f)
+				.Padding(StandardPadding)
+				[
+					SNew(SSpinBox<float>)
+					.Visibility(this, &SFoliageEdit2::GetVisibility_Radius)
+					.MinValue(0.0f)
+					.MaxValue(65536.0f)
+					.MaxSliderValue(8192.0f)
+					.SliderExponent(3.0f)
+					.Value(this, &SFoliageEdit2::GetRadius)
+					.OnValueChanged(this, &SFoliageEdit2::SetRadius)
+				]
+			]
+
+			+ SVerticalBox::Slot()
+			.AutoHeight()
+			[
+				SNew(SHorizontalBox)
+				.ToolTipText(LOCTEXT("PaintDensity_Tooltip", "The density of foliage to paint. This is a multiplier for the individual foliage type's density specifier."))
+
+				+ SHorizontalBox::Slot()
+				.FillWidth(1.0f)
+				.Padding(StandardPadding)
+				[
+					SNew(STextBlock)
+					.Visibility(this, &SFoliageEdit2::GetVisibility_PaintDensity)
+					.Text(LOCTEXT("PaintDensity", "Paint Density"))
+				]
+				+ SHorizontalBox::Slot()
+				.FillWidth(2.0f)
+				.Padding(StandardPadding)
+				[
+					SNew(SSpinBox<float>)
+					.Visibility(this, &SFoliageEdit2::GetVisibility_PaintDensity)
+					.MinValue(0.0f)
+					.MaxValue(2.0f)
+					.MaxSliderValue(1.0f)
+					.Value(this, &SFoliageEdit2::GetPaintDensity)
+					.OnValueChanged(this, &SFoliageEdit2::SetPaintDensity)
+				]
+			]
+
+			+ SVerticalBox::Slot()
+			.AutoHeight()
+			[
+				SNew(SHorizontalBox)
+				.ToolTipText(LOCTEXT("EraseDensity_Tooltip", "The density of foliage to leave behind when erasing with the Shift key held. 0 will remove all foliage."))
+
+				+ SHorizontalBox::Slot()
+				.FillWidth(1.0f)
+				.Padding(StandardPadding)
+				[
+					SNew(STextBlock)
+					.Visibility(this, &SFoliageEdit2::GetVisibility_EraseDensity)
+					.Text(LOCTEXT("EraseDensity", "Erase Density"))
+				]
+				+ SHorizontalBox::Slot()
+				.FillWidth(2.0f)
+				.Padding(StandardPadding)
+				[
+					SNew(SSpinBox<float>)
+					.Visibility(this, &SFoliageEdit2::GetVisibility_EraseDensity)
+					.MinValue(0.0f)
+					.MaxValue(2.0f)
+					.MaxSliderValue(1.0f)
+					.Value(this, &SFoliageEdit2::GetEraseDensity)
+					.OnValueChanged(this, &SFoliageEdit2::SetEraseDensity)
+				]
+			]
+
+			+ SVerticalBox::Slot()
+			.AutoHeight()
+			.Padding(StandardPadding)
+			[
+				SNew(STextBlock)
+				.Visibility(this, &SFoliageEdit2::GetVisibility_Filters)
+				.Text(LOCTEXT("Filter", "Filter"))
+			]
+
+			+ SVerticalBox::Slot()
+			.AutoHeight()
+			.Padding(StandardPadding)
+			[
+				SNew(SWrapBox)
+				.UseAllottedWidth(true)
+
+				+ SWrapBox::Slot()
+				.Padding(0, 0, 16, 0)
+				[
+					SNew(SCheckBox)
+					.Visibility(this, &SFoliageEdit2::GetVisibility_Filters)
+					.OnCheckStateChanged(this, &SFoliageEdit2::OnCheckStateChanged_Landscape)
+					.IsChecked(this, &SFoliageEdit2::GetCheckState_Landscape)
+					.ToolTipText(LOCTEXT("FilterLandscape_Tooltip", "Place foliage on Landscape"))
+					[
+						SNew(STextBlock)
+						.Text(LOCTEXT("Landscape", "Landscape"))
+					]
+				]
+
+				+ SWrapBox::Slot()
+				.Padding(0, 0, 16, 0)
+				[
+					SNew(SCheckBox)
+					.Visibility(this, &SFoliageEdit2::GetVisibility_Filters)
+					.OnCheckStateChanged(this, &SFoliageEdit2::OnCheckStateChanged_StaticMesh)
+					.IsChecked(this, &SFoliageEdit2::GetCheckState_StaticMesh)
+					.ToolTipText(LOCTEXT("FilterStaticMesh_Tooltip", "Place foliage on StaticMeshes"))
+					[
+						SNew(STextBlock)
+						.Text(LOCTEXT("StaticMeshes", "Static Meshes"))
+					]
+				]
+
+				+ SWrapBox::Slot()
+				.Padding(0, 0, 16, 0)
+				[
+					SNew(SCheckBox)
+					.Visibility(this, &SFoliageEdit2::GetVisibility_Filters)
+					.OnCheckStateChanged(this, &SFoliageEdit2::OnCheckStateChanged_BSP)
+					.IsChecked(this, &SFoliageEdit2::GetCheckState_BSP)
+					.ToolTipText(LOCTEXT("FilterBSP_Tooltip", "Place foliage on BSP"))
+					[
+						SNew(STextBlock)
+						.Text(LOCTEXT("BSP", "BSP"))
+					]
+				]
+
+				+ SWrapBox::Slot()
+				.Padding(0, 0, 16, 0)
+				[
+					SNew(SCheckBox)
+					.Visibility(this, &SFoliageEdit2::GetVisibility_Filters)
+					.OnCheckStateChanged(this, &SFoliageEdit2::OnCheckStateChanged_Translucent)
+					.IsChecked(this, &SFoliageEdit2::GetCheckState_Translucent)
+					.ToolTipText(LOCTEXT("FilterTranslucent_Tooltip", "Place foliage on translucent surfaces"))
+					[
+						SNew(STextBlock)
+						.Text(LOCTEXT("Translucent", "Translucent"))
+					]
+				]
+			]
+
+			+ SVerticalBox::Slot()
+			.AutoHeight()
+			[
+				SNew(SHeaderRow)
+
+				+ SHeaderRow::Column(TEXT("Meshes"))
+				.DefaultLabel(LOCTEXT("Meshes", "Meshes"))
+			]
+
+			+ SVerticalBox::Slot()
+			.AutoHeight()
+			.Padding(StandardPadding)
+			[
+				SNew(SBox)
+				.Visibility(this, &SFoliageEdit2::GetVisibility_EmptyList)
+				.Padding(FMargin(15, 0))
+				.MinDesiredHeight(30)
+				[
+					SNew(SScaleBox)
+					.Stretch(EStretch::ScaleToFit)
+					[
+						SNew(STextBlock)
+						.Text(LOCTEXT("Foliage_DropStatic", "+ Drop Static Meshes Here"))
+						.ToolTipText(LOCTEXT("Foliage_DropStatic_ToolTip", "Drag static meshes from the content browser into this area."))
+					]
+				]
+			]
+
+			+ SVerticalBox::Slot()
+			.AutoHeight()
+			[
+				SNew(SBox)
+				.MinDesiredHeight(100)
+				[
+					SAssignNew(MeshTreeWidget, SFoliageMeshTree)
+					.TreeItemsSource(&FoliageEditMode->GetFoliageMeshList())
+					.SelectionMode(ESelectionMode::Multi)
+					.OnGenerateRow(this, &SFoliageEdit2::MeshTreeGenerateRow)
+					.OnGetChildren( this, &SFoliageEdit2::MeshTreeGetChildren)
+					.OnContextMenuOpening(this, &SFoliageEdit2::ConstructFoliageMeshContextMenu)
+					.OnSelectionChanged(this, &SFoliageEdit2::MeshTreeOnSelectionChanged)
+					.HeaderRow
+						(
+							SNew( SHeaderRow ).Visibility(EVisibility::Collapsed)
+							// toggle mesh for painting checkbox
+							+SHeaderRow::Column(FoliageMeshColumns::ColumnID_ToggleMesh).FixedWidth(24.0f)
+							// mesh text label
+							+SHeaderRow::Column(FoliageMeshColumns::ColumnID_MeshLabel).FillWidth(1.0f)
+							// instance count
+							+SHeaderRow::Column(FoliageMeshColumns::ColumnID_InstanceCount).FillWidth(0.3f)
+							// save Foliage asset
+							+SHeaderRow::Column(FoliageMeshColumns::ColumnID_Save).FixedWidth(24.0f)
+						)
+				]
+			]
+			
+			+ SVerticalBox::Slot()
+			.Padding(StandardPadding)
+			[
+				MeshDetailsWidget.ToSharedRef()
+			]
+		]
+
+		// Static Mesh Drop Zone
+		+ SOverlay::Slot()
+		.HAlign(HAlign_Fill)
+		.VAlign(VAlign_Fill)
+		[
+			SNew(SFoliageDragDropHandler2)
+			.Visibility(this, &SFoliageEdit2::GetVisibility_FoliageDropTarget)
+			.FoliageEditPtr(SharedThis(this))
+			[
+				SNew(SScaleBox)
+				.Stretch(EStretch::ScaleToFit)
+				[
+					SNew(STextBlock)
+					.Text(LOCTEXT("Foliage_AddFoliageMesh", "+ Foliage Mesh"))
+					.ShadowOffset(FVector2D(1.f, 1.f))
+				]
+			]
+		]
+	];
+
+	RefreshFullList();
+}
+END_SLATE_FUNCTION_BUILD_OPTIMIZATION
+
+SFoliageEdit2::~SFoliageEdit2()
+{
+	// Release all rendering resources being held onto
+	AssetThumbnailPool->ReleaseResources();
+}
+
+TSharedPtr<FAssetThumbnail> SFoliageEdit2::CreateThumbnail(UStaticMesh* InStaticMesh)
+{
+	return MakeShareable(new FAssetThumbnail(InStaticMesh, 80, 80, AssetThumbnailPool));
+}
+
+void SFoliageEdit2::RefreshFullList()
+{
+	MeshTreeWidget->RequestTreeRefresh();
+}
+
+void SFoliageEdit2::ReplaceItem(const TSharedPtr<SFoliageEditMeshDisplayItem> InDisplayItemToReplaceIn, UStaticMesh* InNewStaticMesh)
+{
+	//TArray<FFoliageMeshUIInfo>& FoliageMeshList = FoliageEditMode->GetFoliageMeshList();
+
+	//for (FFoliageMeshUIInfo& UIInfo : FoliageMeshList)
+	//{
+	//	UFoliageType* Settings = UIInfo.Settings;
+	//	if (Settings->GetStaticMesh() == InNewStaticMesh)
+	//	{
+	//		// This is the info for the new static mesh. Update the display item.
+	//		InDisplayItemToReplaceIn->Replace(Settings,
+	//			CreateThumbnail(Settings->GetStaticMesh()),
+	//			MakeShareable(new FFoliageMeshUIInfo(UIInfo)));
+	//		break;
+	//	}
+	//}
+}
+
+void SFoliageEdit2::ClearAllToolSelection()
+{
+	FoliageEditMode->UISettings.SetLassoSelectToolSelected(false);
+	FoliageEditMode->UISettings.SetPaintToolSelected(false);
+	FoliageEditMode->UISettings.SetReapplyToolSelected(false);
+	FoliageEditMode->UISettings.SetSelectToolSelected(false);
+	FoliageEditMode->UISettings.SetPaintBucketToolSelected(false);
+}
+
+void SFoliageEdit2::BindCommands()
+{
+	const FFoliageEditCommands& Commands = FFoliageEditCommands::Get();
+
+	UICommandList->MapAction(
+		Commands.SetPaint,
+		FExecuteAction::CreateSP(this, &SFoliageEdit2::OnSetPaint),
+		FCanExecuteAction(),
+		FIsActionChecked::CreateSP(this, &SFoliageEdit2::IsPaintTool));
+
+	UICommandList->MapAction(
+		Commands.SetReapplySettings,
+		FExecuteAction::CreateSP(this, &SFoliageEdit2::OnSetReapplySettings),
+		FCanExecuteAction(),
+		FIsActionChecked::CreateSP(this, &SFoliageEdit2::IsReapplySettingsTool));
+
+	UICommandList->MapAction(
+		Commands.SetSelect,
+		FExecuteAction::CreateSP(this, &SFoliageEdit2::OnSetSelectInstance),
+		FCanExecuteAction(),
+		FIsActionChecked::CreateSP(this, &SFoliageEdit2::IsSelectTool));
+
+	UICommandList->MapAction(
+		Commands.SetLassoSelect,
+		FExecuteAction::CreateSP(this, &SFoliageEdit2::OnSetLasso),
+		FCanExecuteAction(),
+		FIsActionChecked::CreateSP(this, &SFoliageEdit2::IsLassoSelectTool));
+
+	UICommandList->MapAction(
+		Commands.SetPaintBucket,
+		FExecuteAction::CreateSP(this, &SFoliageEdit2::OnSetPaintFill),
+		FCanExecuteAction(),
+		FIsActionChecked::CreateSP(this, &SFoliageEdit2::IsPaintFillTool));
+
+		//
+	UICommandList->MapAction(
+		Commands.RemoveFoliageType,
+		FExecuteAction::CreateSP(this, &SFoliageEdit2::OnRemoveFoliageType),
+		FCanExecuteAction(),
+		FIsActionChecked());
+
+	UICommandList->MapAction(
+		Commands.ShowFoliageTypeInCB,
+		FExecuteAction::CreateSP(this, &SFoliageEdit2::OnShowFoliageTypeInCB),
+		FCanExecuteAction(),
+		FIsActionChecked());
+
+}
+
+TSharedRef<SWidget> SFoliageEdit2::BuildToolBar()
+{
+	FToolBarBuilder Toolbar(UICommandList, FMultiBoxCustomization::None);
+	{
+		Toolbar.AddToolBarButton(FFoliageEditCommands::Get().SetPaint);
+		Toolbar.AddToolBarButton(FFoliageEditCommands::Get().SetReapplySettings);
+		Toolbar.AddToolBarButton(FFoliageEditCommands::Get().SetSelect);
+		Toolbar.AddToolBarButton(FFoliageEditCommands::Get().SetLassoSelect);
+		Toolbar.AddToolBarButton(FFoliageEditCommands::Get().SetPaintBucket);
+	}
+
+	IIntroTutorials& IntroTutorials = FModuleManager::LoadModuleChecked<IIntroTutorials>(TEXT("IntroTutorials"));
+
+	return
+		SNew(SHorizontalBox)
+
+		+ SHorizontalBox::Slot()
+		.Padding(4, 0)
+		[
+			SNew(SOverlay)
+			+ SOverlay::Slot()
+			[
+				SNew(SBorder)
+				.HAlign(HAlign_Center)
+				.Padding(0)
+				.BorderImage(FEditorStyle::GetBrush("NoBorder"))
+				.IsEnabled(FSlateApplication::Get().GetNormalExecutionAttribute())
+				[
+					Toolbar.MakeWidget()
+				]
+			]
+
+			// Tutorial link
+			+ SOverlay::Slot()
+			.HAlign(HAlign_Right)
+			.VAlign(VAlign_Bottom)
+			.Padding(4)
+			[
+				IntroTutorials.CreateTutorialsWidget(TEXT("FoliageMode"))
+			]
+		];
+}
+
+void SFoliageEdit2::OnSetPaint()
+{
+	ClearAllToolSelection();
+
+	FoliageEditMode->UISettings.SetPaintToolSelected(true);
+
+	FoliageEditMode->NotifyToolChanged();
+}
+
+bool SFoliageEdit2::IsPaintTool() const
+{
+	return FoliageEditMode->UISettings.GetPaintToolSelected();
+}
+
+void SFoliageEdit2::OnSetReapplySettings()
+{
+	ClearAllToolSelection();
+
+	FoliageEditMode->UISettings.SetReapplyToolSelected(true);
+
+	FoliageEditMode->NotifyToolChanged();
+}
+
+bool SFoliageEdit2::IsReapplySettingsTool() const
+{
+	return FoliageEditMode->UISettings.GetReapplyToolSelected();
+}
+
+void SFoliageEdit2::OnSetSelectInstance()
+{
+	ClearAllToolSelection();
+
+	FoliageEditMode->UISettings.SetSelectToolSelected(true);
+
+	FoliageEditMode->NotifyToolChanged();
+}
+
+bool SFoliageEdit2::IsSelectTool() const
+{
+	return FoliageEditMode->UISettings.GetSelectToolSelected();
+}
+
+void SFoliageEdit2::OnSetLasso()
+{
+	ClearAllToolSelection();
+
+	FoliageEditMode->UISettings.SetLassoSelectToolSelected(true);
+
+	FoliageEditMode->NotifyToolChanged();
+}
+
+bool SFoliageEdit2::IsLassoSelectTool() const
+{
+	return FoliageEditMode->UISettings.GetLassoSelectToolSelected();
+}
+
+void SFoliageEdit2::OnSetPaintFill()
+{
+	ClearAllToolSelection();
+
+	FoliageEditMode->UISettings.SetPaintBucketToolSelected(true);
+
+	FoliageEditMode->NotifyToolChanged();
+}
+
+bool SFoliageEdit2::IsPaintFillTool() const
+{
+	return FoliageEditMode->UISettings.GetPaintBucketToolSelected();
+}
+
+void SFoliageEdit2::SetRadius(float InRadius)
+{
+	FoliageEditMode->UISettings.SetRadius(InRadius);
+}
+
+float SFoliageEdit2::GetRadius() const
+{
+	return FoliageEditMode->UISettings.GetRadius();
+}
+
+void SFoliageEdit2::SetPaintDensity(float InDensity)
+{
+	FoliageEditMode->UISettings.SetPaintDensity(InDensity);
+}
+
+float SFoliageEdit2::GetPaintDensity() const
+{
+	return FoliageEditMode->UISettings.GetPaintDensity();
+}
+
+
+void SFoliageEdit2::SetEraseDensity(float InDensity)
+{
+	FoliageEditMode->UISettings.SetUnpaintDensity(InDensity);
+}
+
+float SFoliageEdit2::GetEraseDensity() const
+{
+	return FoliageEditMode->UISettings.GetUnpaintDensity();
+}
+
+EVisibility SFoliageEdit2::GetVisibility_EmptyList() const
+{
+	if (FoliageEditMode->GetFoliageMeshList().Num() > 0)
+	{
+		return EVisibility::Collapsed;
+	}
+
+	return EVisibility::HitTestInvisible;
+}
+
+EVisibility SFoliageEdit2::GetVisibility_FoliageDropTarget() const
+{
+	if ( FSlateApplication::Get().IsDragDropping() )
+	{
+		TArray<FAssetData> DraggedAssets = AssetUtil::ExtractAssetDataFromDrag(FSlateApplication::Get().GetDragDroppingContent());
+		for ( const FAssetData& AssetData : DraggedAssets )
+		{
+			if ( AssetData.IsValid() && (AssetData.GetClass()->IsChildOf(UStaticMesh::StaticClass()) ||	AssetData.GetClass()->IsChildOf(UFoliageType::StaticClass())) )
+			{
+				return EVisibility::Visible;
+			}
+		}
+	}
+	
+	return EVisibility::Hidden;
+}
+
+EVisibility SFoliageEdit2::GetVisibility_NonEmptyList() const
+{
+	if (FoliageEditMode->GetFoliageMeshList().Num() == 0)
+	{
+		return EVisibility::Collapsed;
+	}
+
+	return EVisibility::HitTestInvisible;
+}
+
+FReply SFoliageEdit2::OnDrop_ListView(const FGeometry& MyGeometry, const FDragDropEvent& DragDropEvent)
+{
+	TArray<FAssetData> DroppedAssetData = AssetUtil::ExtractAssetDataFromDrag(DragDropEvent);
+
+	for (int32 AssetIdx = 0; AssetIdx < DroppedAssetData.Num(); ++AssetIdx)
+	{
+		GWarn->BeginSlowTask(LOCTEXT("OnDrop_LoadPackage", "Fully Loading Package For Drop"), true, false);
+		UObject* Asset = DroppedAssetData[AssetIdx].GetAsset();
+		GWarn->EndSlowTask();
+
+		UFoliageType* FoliageType = FoliageEditMode->AddFoliageAsset(Asset);
+		if (FoliageType)
+		{
+			MeshTreeWidget->RequestTreeRefresh();
+		}
+	}
+
+	return FReply::Handled();
+}
+TSharedPtr<SWidget> SFoliageEdit2::ConstructFoliageMeshContextMenu() const
+{
+	const FFoliageEditCommands& Commands = FFoliageEditCommands::Get();
+	
+	FMenuBuilder MenuBuilder(true, UICommandList);
+	MenuBuilder.AddMenuEntry(Commands.RemoveFoliageType);
+
+	MenuBuilder.AddSubMenu(
+		LOCTEXT( "ReplaceFoliageType", "Replace Foliage Type" ),
+		LOCTEXT ("ReplaceFoliageType_ToolTip", "Replaces selected foliage type with another foliage type asset"),
+		FNewMenuDelegate::CreateSP(this, &SFoliageEdit2::FillReplaceFoliageTypeSubmenu));
+	
+	MenuBuilder.AddMenuEntry(Commands.ShowFoliageTypeInCB);
+
+	//// Add common commands
+	//MenuBuilder.BeginSection("Levels", LOCTEXT("LevelsHeader", "Levels"));
+	//{
+	//	// Visibility commands
+	//	MenuBuilder.AddSubMenu(
+	//		LOCTEXT("VisibilityHeader", "Visibility"),
+	//		LOCTEXT("VisibilitySubMenu_ToolTip", "Selected Level(s) visibility commands"),
+	//		FNewMenuDelegate::CreateSP(this, &FStreamingLevelCollectionModel::FillVisibilitySubMenu));
+
+	//	// Lock commands
+	//	MenuBuilder.AddSubMenu(
+	//		LOCTEXT("LockHeader", "Lock"),
+	//		LOCTEXT("LockSubMenu_ToolTip", "Selected Level(s) lock commands"),
+	//		FNewMenuDelegate::CreateSP(this, &FStreamingLevelCollectionModel::FillLockSubMenu));
+	//}
+	//MenuBuilder.EndSection();
+
+	return MenuBuilder.MakeWidget();
+}
+
+void SFoliageEdit2::OnCheckStateChanged_Landscape(ECheckBoxState InState)
+{
+	FoliageEditMode->UISettings.SetFilterLandscape(InState == ECheckBoxState::Checked ? true : false);
+}
+
+ECheckBoxState SFoliageEdit2::GetCheckState_Landscape() const
+{
+	return FoliageEditMode->UISettings.GetFilterLandscape() ? ECheckBoxState::Checked : ECheckBoxState::Unchecked;
+}
+
+void SFoliageEdit2::OnCheckStateChanged_StaticMesh(ECheckBoxState InState)
+{
+	FoliageEditMode->UISettings.SetFilterStaticMesh(InState == ECheckBoxState::Checked ? true : false);
+}
+
+ECheckBoxState SFoliageEdit2::GetCheckState_StaticMesh() const
+{
+	return FoliageEditMode->UISettings.GetFilterStaticMesh() ? ECheckBoxState::Checked : ECheckBoxState::Unchecked;
+}
+
+void SFoliageEdit2::OnCheckStateChanged_BSP(ECheckBoxState InState)
+{
+	FoliageEditMode->UISettings.SetFilterBSP(InState == ECheckBoxState::Checked ? true : false);
+}
+
+ECheckBoxState SFoliageEdit2::GetCheckState_BSP() const
+{
+	return FoliageEditMode->UISettings.GetFilterBSP() ? ECheckBoxState::Checked : ECheckBoxState::Unchecked;
+}
+
+void SFoliageEdit2::OnCheckStateChanged_Translucent(ECheckBoxState InState)
+{
+	FoliageEditMode->UISettings.SetFilterTranslucent(InState == ECheckBoxState::Checked ? true : false);
+}
+
+ECheckBoxState SFoliageEdit2::GetCheckState_Translucent() const
+{
+	return FoliageEditMode->UISettings.GetFilterTranslucent() ? ECheckBoxState::Checked : ECheckBoxState::Unchecked;
+}
+
+EVisibility SFoliageEdit2::GetVisibility_Radius() const
+{
+	if (FoliageEditMode->UISettings.GetSelectToolSelected() || FoliageEditMode->UISettings.GetReapplyPaintBucketToolSelected() || FoliageEditMode->UISettings.GetPaintBucketToolSelected() )
+	{
+		return EVisibility::Collapsed;
+	}
+
+	return EVisibility::Visible;
+}
+
+EVisibility SFoliageEdit2::GetVisibility_PaintDensity() const
+{
+	if (FoliageEditMode->UISettings.GetSelectToolSelected() || FoliageEditMode->UISettings.GetReapplyPaintBucketToolSelected() || FoliageEditMode->UISettings.GetLassoSelectToolSelected())
+	{
+		return EVisibility::Collapsed;
+	}
+
+	return EVisibility::Visible;
+}
+
+EVisibility SFoliageEdit2::GetVisibility_EraseDensity() const
+{
+	if (FoliageEditMode->UISettings.GetSelectToolSelected() || FoliageEditMode->UISettings.GetReapplyPaintBucketToolSelected() || FoliageEditMode->UISettings.GetLassoSelectToolSelected() || FoliageEditMode->UISettings.GetPaintBucketToolSelected())
+	{
+		return EVisibility::Collapsed;
+	}
+
+	return EVisibility::Visible;
+}
+
+EVisibility SFoliageEdit2::GetVisibility_Filters() const
+{
+	if (FoliageEditMode->UISettings.GetSelectToolSelected())
+	{
+		return EVisibility::Collapsed;
+	}
+
+	return EVisibility::Visible;
+}
+
+TSharedRef<ITableRow> SFoliageEdit2::MeshTreeGenerateRow(FFoliageMeshUIInfoPtr Item, const TSharedRef<STableViewBase>& OwnerTable)
+{
+	return SNew(SFoliageEditMeshRow, OwnerTable)
+		.MeshInfo(Item)
+		.FoliageEditMode(FoliageEditMode);
+}
+
+void SFoliageEdit2::MeshTreeGetChildren(FFoliageMeshUIInfoPtr Item, TArray<FFoliageMeshUIInfoPtr>& OutChildren)
+{
+	//OutChildren = Item->GetChildren();
+}
+
+void SFoliageEdit2::MeshTreeOnSelectionChanged(FFoliageMeshUIInfoPtr Item, ESelectInfo::Type SelectInfo)
+{
+	TArray<UObject*> FoliageMeshes;
+	auto InfoList = MeshTreeWidget->GetSelectedItems();
+
+	for (const auto& Info : InfoList)
+	{
+		FoliageMeshes.Add(Info->Settings);
+	}
+
+	MeshDetailsWidget->SetObjects(FoliageMeshes, true);
+}
+
+void SFoliageEdit2::OnRemoveFoliageType()
+{
+	int32 NumInstances = 0;
+	TArray<UFoliageType*> FoliageTypeList;
+	
+	auto MeshUIList = MeshTreeWidget->GetSelectedItems();
+	for (FFoliageMeshUIInfoPtr& MeshUIInfo : MeshUIList)
+	{
+		NumInstances+= MeshUIInfo->InstanceCountTotal;
+		FoliageTypeList.Add(MeshUIInfo->Settings);
+	}
+	
+	bool bProceed = true;
+	if (NumInstances > 0)
+	{
+		FText Message = FText::Format(NSLOCTEXT("UnrealEd", "FoliageMode_DeleteMesh", "Are you sure you want to remove {0} instances?"), FText::AsNumber(NumInstances));
+		bProceed = (FMessageDialog::Open(EAppMsgType::YesNo, Message) == EAppReturnType::Yes);
+	}
+
+	if (bProceed)
+	{
+		FoliageEditMode->RemoveFoliageType(FoliageTypeList.GetData(), FoliageTypeList.Num());
+	}
+}
+
+void SFoliageEdit2::FillReplaceFoliageTypeSubmenu(FMenuBuilder& MenuBuilder)
+{
+	FContentBrowserModule& ContentBrowserModule = FModuleManager::Get().LoadModuleChecked<FContentBrowserModule>(TEXT("ContentBrowser"));
+
+	FAssetPickerConfig AssetPickerConfig;
+	AssetPickerConfig.Filter.ClassNames.Add(UFoliageType::StaticClass()->GetFName());
+	AssetPickerConfig.Filter.bRecursiveClasses = true;
+	AssetPickerConfig.OnAssetSelected = FOnAssetSelected::CreateSP(this, &SFoliageEdit2::OnReplaceFoliageTypeSelected);
+	AssetPickerConfig.InitialAssetViewType = EAssetViewType::List;
+	AssetPickerConfig.bAllowNullSelection = false;
+	
+	TSharedRef<SWidget> MenuContent = SNew(SBox)
+	.WidthOverride(384)
+	.HeightOverride(500)
+	[
+		ContentBrowserModule.Get().CreateAssetPicker(AssetPickerConfig)
+	];
+	
+	MenuBuilder.AddWidget( MenuContent, FText::GetEmpty(), true);
+}
+
+void SFoliageEdit2::OnReplaceFoliageTypeSelected(const FAssetData& AssetData)
+{
+	FSlateApplication::Get().DismissAllMenus();
+	
+	auto MeshUIList = MeshTreeWidget->GetSelectedItems();
+	if (MeshUIList.Num())
+	{
+		UFoliageType* OldSettings = MeshUIList[0]->Settings;
+
+		UFoliageType* NewSettings = Cast<UFoliageType>(AssetData.GetAsset());
+		if (NewSettings && OldSettings != NewSettings)
+		{
+			FoliageEditMode->ReplaceSettingsObject(OldSettings, NewSettings);
+		}
+	}
+}
+
+void SFoliageEdit2::OnShowFoliageTypeInCB()
+{
+	TArray<UObject*> SelectedAssets;
+	auto MeshUIList = MeshTreeWidget->GetSelectedItems();
+	for (FFoliageMeshUIInfoPtr  MeshUIPtr : MeshUIList)
+	{
+		UFoliageType* FoliageType = MeshUIPtr->Settings;
+		if (FoliageType->IsAsset())		
+		{
+			SelectedAssets.Add(FoliageType);
+		}
+		else
+		{
+			SelectedAssets.Add(FoliageType->GetStaticMesh());
+		}
+	}
+
+	if (SelectedAssets.Num())
+	{
+		GEditor->SyncBrowserToObjects(SelectedAssets);
+	}
+}
+
+#undef LOCTEXT_NAMESPACE
diff --git a/Engine/Source/Editor/FoliageEdit/Private/SFoliageEdit2.h b/Engine/Source/Editor/FoliageEdit/Private/SFoliageEdit2.h
new file mode 100644
--- /dev/null
+++ b/Engine/Source/Editor/FoliageEdit/Private/SFoliageEdit2.h
@@ -0,0 +1,179 @@
+// Copyright 1998-2015 Epic Games, Inc. All Rights Reserved.
+
+#pragma once
+
+// Forwards declarations
+class FAssetThumbnail;
+class SFoliageEditMeshDisplayItem;
+struct FFoliageMeshUIInfo;
+class FAssetThumbnailPool;
+class IDetailsView;
+
+typedef STreeView<FFoliageMeshUIInfoPtr> SFoliageMeshTree;
+
+class SFoliageEdit2 : public SCompoundWidget
+{
+public:
+	SLATE_BEGIN_ARGS(SFoliageEdit2) {}
+	SLATE_END_ARGS()
+
+public:
+	/** SCompoundWidget functions */
+	void Construct(const FArguments& InArgs);
+
+	~SFoliageEdit2();
+
+	/** Creates the thumbnail for the passed in Static Mesh. Used by the MeshListView items. */
+	TSharedPtr<FAssetThumbnail> CreateThumbnail(UStaticMesh* InStaticMesh);
+
+	/** Does a full refresh on the list. */
+	void RefreshFullList();
+
+	void ReplaceItem(const TSharedPtr<SFoliageEditMeshDisplayItem> InDisplayItemToReplaceIn, UStaticMesh* InNewStaticMesh);
+
+	/** Handles adding a new item to the list and refreshing the list in it's entirety. */
+	FReply OnDrop_ListView(const FGeometry& MyGeometry, const FDragDropEvent& DragDropEvent);
+
+	/** Gets FoliageEditMode. Used by the cluster details to notify changes */
+	class FEdModeFoliage* GetFoliageEditMode() const { return FoliageEditMode; }
+
+	/** @return the SWidget containing the context menu */
+	TSharedPtr<SWidget> ConstructFoliageMeshContextMenu() const;
+
+private:
+	/** Clears all the tools selection by setting them to false. */
+	void ClearAllToolSelection();
+
+	/** Binds UI commands for the toolbar. */
+	void BindCommands();
+
+	/** Creates the toolbar. */
+	TSharedRef<SWidget> BuildToolBar();
+
+	/** Delegate callbacks for the UI */
+
+	/** Sets the tool mode to Paint. */
+	void OnSetPaint();
+
+	/** Checks if the tool mode is Paint. */
+	bool IsPaintTool() const;
+
+	/** Sets the tool mode to Reapply Settings. */
+	void OnSetReapplySettings();
+
+	/** Checks if the tool mode is Reapply Settings. */
+	bool IsReapplySettingsTool() const;
+
+	/** Sets the tool mode to Select. */
+	void OnSetSelectInstance();
+
+	/** Checks if the tool mode is Select. */
+	bool IsSelectTool() const;
+
+	/** Sets the tool mode to Lasso Select. */
+	void OnSetLasso();
+
+	/** Checks if the tool mode is Lasso Select. */
+	bool IsLassoSelectTool() const;
+
+	/** Sets the tool mode to Paint Bucket. */
+	void OnSetPaintFill();
+
+	/** Checks if the tool mode is Paint Bucket. */
+	bool IsPaintFillTool() const;
+
+	/** Sets the brush Radius for the brush. */
+	void SetRadius(float InRadius);
+
+	/** Retrieves the brush Radius for the brush. */
+	float GetRadius() const;
+
+	/** Sets the Paint Density for the brush. */
+	void SetPaintDensity(float InDensity);
+
+	/** Retrieves the Paint Density for the brush. */
+	float GetPaintDensity() const;
+
+	/** Sets the Erase Density for the brush. */
+	void SetEraseDensity(float InDensity);
+
+	/** Retrieves the Erase Density for the brush. */
+	float GetEraseDensity() const;
+
+	/** Creates the list item widget that displays the instance settings. */
+	TSharedRef<ITableRow> MakeWidgetFromOption(FFoliageMeshUIInfoPtr InItem, const TSharedRef<STableViewBase>& OwnerTable);
+
+	/** Sets the filter settings for if painting will occur on Landscapes. */
+	void OnCheckStateChanged_Landscape(ECheckBoxState InState);
+
+	/** Retrieves the filter settings for painting on Landscapes. */
+	ECheckBoxState GetCheckState_Landscape() const;
+
+	/** Sets the filter settings for if painting will occur on Static Meshes. */
+	void OnCheckStateChanged_StaticMesh(ECheckBoxState InState);
+
+	/** Retrieves the filter settings for painting on Static Meshes. */
+	ECheckBoxState GetCheckState_StaticMesh() const;
+
+	/** Sets the filter settings for if painting will occur on BSPs. */
+	void OnCheckStateChanged_BSP(ECheckBoxState InState);
+
+	/** Retrieves the filter settings for painting on BSPs. */
+	ECheckBoxState GetCheckState_BSP() const;
+
+	/** Sets the filter settings for if painting will occur on translucent meshes. */
+	void OnCheckStateChanged_Translucent(ECheckBoxState InState);
+
+	/** Retrieves the filter settings for painting on translucent meshes. */
+	ECheckBoxState GetCheckState_Translucent() const;
+
+	/** Checks if the text in the empty list overlay should appear. If the list is has items but the the drag and drop override is true, it will return EVisibility::Visible. */
+	EVisibility GetVisibility_EmptyList() const;
+
+	/** Should the drop area be visible currently?  Happens when the user is dragging static meshes */
+	EVisibility GetVisibility_FoliageDropTarget() const;
+
+	/** Checks if the list should appear. */
+	EVisibility GetVisibility_NonEmptyList() const;
+
+	/** Checks if the radius spinbox should appear. Dependant on the current tool being used. */
+	EVisibility GetVisibility_Radius() const;
+
+	/** Checks if the paint density spinbox should appear. Dependant on the current tool being used. */
+	EVisibility GetVisibility_PaintDensity() const;
+
+	/** Checks if the erase density spinbox should appear. Dependant on the current tool being used. */
+	EVisibility GetVisibility_EraseDensity() const;
+
+	/** Checks if the filters should appear. Dependant on the current tool being used. */
+	EVisibility GetVisibility_Filters() const;
+
+	/**  */
+	TSharedRef<ITableRow> MeshTreeGenerateRow(FFoliageMeshUIInfoPtr Item, const TSharedRef<STableViewBase>& OwnerTable);
+	
+	/**  */
+	void MeshTreeGetChildren(FFoliageMeshUIInfoPtr Item, TArray<FFoliageMeshUIInfoPtr>& OutChildren);
+
+	/**  */
+	void MeshTreeOnSelectionChanged(FFoliageMeshUIInfoPtr Item, ESelectInfo::Type SelectInfo);
+
+	void FillReplaceFoliageTypeSubmenu(FMenuBuilder& MenuBuilder);
+	/***/
+	void OnRemoveFoliageType();
+	void OnShowFoliageTypeInCB();
+	void OnReplaceFoliageTypeSelected(const class FAssetData& AssetData);
+		
+
+private:
+	TSharedPtr<SFoliageMeshTree>	MeshTreeWidget;
+	TSharedPtr<IDetailsView>		MeshDetailsWidget;
+	
+	/** Pool for maintaining and rendering thumbnails */
+	TSharedPtr<FAssetThumbnailPool> AssetThumbnailPool;
+
+	/** Command list for binding functions for the toolbar. */
+	TSharedPtr<FUICommandList> UICommandList;
+
+	/** Pointer to the foliage edit mode. */
+	FEdModeFoliage* FoliageEditMode;
+};
diff --git a/Engine/Source/Editor/FoliageEdit/Private/SFoliageEditMeshDisplayItem.cpp b/Engine/Source/Editor/FoliageEdit/Private/SFoliageEditMeshDisplayItem.cpp
--- a/Engine/Source/Editor/FoliageEdit/Private/SFoliageEditMeshDisplayItem.cpp
+++ b/Engine/Source/Editor/FoliageEdit/Private/SFoliageEditMeshDisplayItem.cpp
@@ -1611,7 +1611,7 @@ void SFoliageEditMeshDisplayItem::Construct(const FArguments& InArgs)
 						SNew(SButton)
 						.ButtonStyle(FEditorStyle::Get(), "NoBorder")
 						.OnClicked(this, &SFoliageEditMeshDisplayItem::OnSaveSettings)
-						.Visibility(this, &SFoliageEditMeshDisplayItem::IsSaveSettingsVisible)
+						.IsEnabled(this, &SFoliageEditMeshDisplayItem::IsSaveSettingsEnabled)
 						.ToolTipText(this, &SFoliageEditMeshDisplayItem::GetSaveSettingsTooltip)
 						.AddMetaData<FTutorialMetaData>(FTutorialMetaData(TEXT("Foliage.SaveSettings"), "LevelEditorToolbox"))
 						[
@@ -1911,9 +1911,14 @@ FText SFoliageEditMeshDisplayItem::GetSaveSettingsTooltip() const
 	return FText();
 }
 
-EVisibility SFoliageEditMeshDisplayItem::IsSaveSettingsVisible() const
+bool SFoliageEditMeshDisplayItem::IsSaveSettingsEnabled() const
 {
-	return !FoliageSettingsPtr->IsAsset() ? EVisibility::Visible : EVisibility::Collapsed;
+	if (!FoliageSettingsPtr->IsAsset() || FoliageSettingsPtr->GetOutermost()->IsDirty())
+	{
+		return true;
+	}
+
+	return false;
 }
 
 bool SFoliageEditMeshDisplayItem::IsPropertyVisible(const FPropertyAndParent& PropertyAndParent) const
@@ -1992,6 +1997,7 @@ ECheckBoxState SFoliageEditMeshDisplayItem::IsDensityReapplyChecked() const
 
 void SFoliageEditMeshDisplayItem::OnDensityChanged(float InValue)
 {
+	FoliageSettingsPtr->Modify();
 	FoliageSettingsPtr->Density = InValue;
 }
 
@@ -2002,6 +2008,7 @@ float SFoliageEditMeshDisplayItem::GetDensity() const
 
 void SFoliageEditMeshDisplayItem::OnDensityReapplyChanged(float InValue)
 {
+	FoliageSettingsPtr->Modify();
 	FoliageSettingsPtr->ReapplyDensityAmount = InValue;
 }
 
@@ -2012,6 +2019,7 @@ float SFoliageEditMeshDisplayItem::GetDensityReapply() const
 
 void SFoliageEditMeshDisplayItem::OnRadiusReapply(ECheckBoxState InState)
 {
+	FoliageSettingsPtr->Modify();
 	FoliageSettingsPtr->ReapplyRadius = InState == ECheckBoxState::Checked ? true : false;
 }
 
@@ -2022,6 +2030,7 @@ ECheckBoxState SFoliageEditMeshDisplayItem::IsRadiusReapplyChecked() const
 
 void SFoliageEditMeshDisplayItem::OnRadiusChanged(float InValue)
 {
+	FoliageSettingsPtr->Modify();
 	FoliageSettingsPtr->Radius = InValue;
 }
 
@@ -2032,6 +2041,7 @@ float SFoliageEditMeshDisplayItem::GetRadius() const
 
 void SFoliageEditMeshDisplayItem::OnAlignToNormalReapply(ECheckBoxState InState)
 {
+	FoliageSettingsPtr->Modify();
 	FoliageSettingsPtr->ReapplyAlignToNormal = InState == ECheckBoxState::Checked ? true : false;
 }
 
@@ -2042,6 +2052,7 @@ ECheckBoxState SFoliageEditMeshDisplayItem::IsAlignToNormalReapplyChecked() cons
 
 void SFoliageEditMeshDisplayItem::OnAlignToNormal(ECheckBoxState InState)
 {
+	FoliageSettingsPtr->Modify();
 	if (InState == ECheckBoxState::Checked)
 	{
 		FoliageSettingsPtr->AlignToNormal = true;
@@ -2064,6 +2075,7 @@ EVisibility SFoliageEditMeshDisplayItem::IsAlignToNormalVisible() const
 
 void SFoliageEditMeshDisplayItem::OnMaxAngleChanged(float InValue)
 {
+	FoliageSettingsPtr->Modify();
 	FoliageSettingsPtr->AlignMaxAngle = InValue;
 }
 
@@ -2074,6 +2086,7 @@ float SFoliageEditMeshDisplayItem::GetMaxAngle() const
 
 void SFoliageEditMeshDisplayItem::OnRandomYawReapply(ECheckBoxState InState)
 {
+	FoliageSettingsPtr->Modify();
 	FoliageSettingsPtr->ReapplyRandomYaw = InState == ECheckBoxState::Checked ? true : false;
 }
 
@@ -2084,6 +2097,7 @@ ECheckBoxState SFoliageEditMeshDisplayItem::IsRandomYawReapplyChecked() const
 
 void SFoliageEditMeshDisplayItem::OnRandomYaw(ECheckBoxState InState)
 {
+	FoliageSettingsPtr->Modify();
 	if (InState == ECheckBoxState::Checked)
 	{
 		FoliageSettingsPtr->RandomYaw = true;
@@ -2101,12 +2115,13 @@ ECheckBoxState SFoliageEditMeshDisplayItem::IsRandomYawChecked() const
 
 void SFoliageEditMeshDisplayItem::OnUniformScale(ECheckBoxState InState)
 {
+	FoliageSettingsPtr->Modify();
 	if (InState == ECheckBoxState::Checked)
 	{
 		FoliageSettingsPtr->UniformScale = true;
 
-		FoliageSettingsPtr->ScaleMinY = FoliageSettingsPtr->ScaleMinX;
-		FoliageSettingsPtr->ScaleMinZ = FoliageSettingsPtr->ScaleMinX;
+		FoliageSettingsPtr->ScaleY = FoliageSettingsPtr->ScaleX;
+		FoliageSettingsPtr->ScaleZ = FoliageSettingsPtr->ScaleX;
 	}
 	else
 	{
@@ -2121,6 +2136,7 @@ ECheckBoxState SFoliageEditMeshDisplayItem::IsUniformScaleChecked() const
 
 void SFoliageEditMeshDisplayItem::OnScaleUniformReapply(ECheckBoxState InState)
 {
+	FoliageSettingsPtr->Modify();
 	FoliageSettingsPtr->ReapplyScaleX = InState == ECheckBoxState::Checked ? true : false;
 	FoliageSettingsPtr->ReapplyScaleY = InState == ECheckBoxState::Checked ? true : false;
 	FoliageSettingsPtr->ReapplyScaleZ = InState == ECheckBoxState::Checked ? true : false;
@@ -2133,38 +2149,37 @@ ECheckBoxState SFoliageEditMeshDisplayItem::IsScaleUniformReapplyChecked() const
 
 void SFoliageEditMeshDisplayItem::OnScaleUniformMinChanged(float InValue)
 {
-	FoliageSettingsPtr->ScaleMinX = InValue;
-	FoliageSettingsPtr->ScaleMinY = InValue;
-	FoliageSettingsPtr->ScaleMinZ = InValue;
-
-	FoliageSettingsPtr->ScaleMaxX = FMath::Max(FoliageSettingsPtr->ScaleMaxX, InValue);
-	FoliageSettingsPtr->ScaleMaxY = FoliageSettingsPtr->ScaleMaxX;
-	FoliageSettingsPtr->ScaleMaxZ = FoliageSettingsPtr->ScaleMaxX;
+	FoliageSettingsPtr->Modify();
+	FoliageSettingsPtr->ScaleX.Min = InValue;
+	FoliageSettingsPtr->ScaleX.Include(InValue);
+	
+	FoliageSettingsPtr->ScaleY = FoliageSettingsPtr->ScaleX;
+	FoliageSettingsPtr->ScaleZ = FoliageSettingsPtr->ScaleX;
 }
 
 float SFoliageEditMeshDisplayItem::GetScaleUniformMin() const
 {
-	return FoliageSettingsPtr->ScaleMinX;
+	return FoliageSettingsPtr->ScaleX.Min;
 }
 
 void SFoliageEditMeshDisplayItem::OnScaleUniformMaxChanged(float InValue)
 {
-	FoliageSettingsPtr->ScaleMaxX = InValue;
-	FoliageSettingsPtr->ScaleMaxY = InValue;
-	FoliageSettingsPtr->ScaleMaxZ = InValue;
+	FoliageSettingsPtr->Modify();
+	FoliageSettingsPtr->ScaleX.Max = InValue;
+	FoliageSettingsPtr->ScaleX.Include(InValue);
 
-	FoliageSettingsPtr->ScaleMinX = FMath::Min(FoliageSettingsPtr->ScaleMinX, InValue);
-	FoliageSettingsPtr->ScaleMinY = FoliageSettingsPtr->ScaleMinX;
-	FoliageSettingsPtr->ScaleMinZ = FoliageSettingsPtr->ScaleMinX;
+	FoliageSettingsPtr->ScaleY = FoliageSettingsPtr->ScaleX;
+	FoliageSettingsPtr->ScaleZ = FoliageSettingsPtr->ScaleX;
 }
 
 float SFoliageEditMeshDisplayItem::GetScaleUniformMax() const
 {
-	return FoliageSettingsPtr->ScaleMaxX;
+	return FoliageSettingsPtr->ScaleX.Max;
 }
 
 void SFoliageEditMeshDisplayItem::OnScaleXReapply(ECheckBoxState InState)
 {
+	FoliageSettingsPtr->Modify();
 	FoliageSettingsPtr->ReapplyScaleX = InState == ECheckBoxState::Checked;
 }
 
@@ -2175,26 +2190,26 @@ ECheckBoxState SFoliageEditMeshDisplayItem::IsScaleXReapplyChecked() const
 
 void SFoliageEditMeshDisplayItem::OnScaleXMinChanged(float InValue)
 {
-	FoliageSettingsPtr->ScaleMinX = InValue;
-
-	FoliageSettingsPtr->ScaleMaxX = FMath::Max(FoliageSettingsPtr->ScaleMaxX, InValue);
+	FoliageSettingsPtr->Modify();
+	FoliageSettingsPtr->ScaleX.Min = InValue;
+	FoliageSettingsPtr->ScaleX.Include(InValue);
 }
 
 float SFoliageEditMeshDisplayItem::GetScaleXMin() const
 {
-	return FoliageSettingsPtr->ScaleMinX;
+	return FoliageSettingsPtr->ScaleX.Min;
 }
 
 void SFoliageEditMeshDisplayItem::OnScaleXMaxChanged(float InValue)
 {
-	FoliageSettingsPtr->ScaleMaxX = InValue;
-
-	FoliageSettingsPtr->ScaleMinX = FMath::Min(FoliageSettingsPtr->ScaleMinX, InValue);
+	FoliageSettingsPtr->Modify();
+	FoliageSettingsPtr->ScaleX.Max = InValue;
+	FoliageSettingsPtr->ScaleX.Include(InValue);
 }
 
 float SFoliageEditMeshDisplayItem::GetScaleXMax() const
 {
-	return FoliageSettingsPtr->ScaleMaxX;
+	return FoliageSettingsPtr->ScaleX.Max;
 }
 
 ECheckBoxState SFoliageEditMeshDisplayItem::IsScaleXLockedChecked() const
@@ -2204,11 +2219,13 @@ ECheckBoxState SFoliageEditMeshDisplayItem::IsScaleXLockedChecked() const
 
 void SFoliageEditMeshDisplayItem::OnScaleXLocked(ECheckBoxState InState)
 {
+	FoliageSettingsPtr->Modify();
 	FoliageSettingsPtr->LockScaleX = InState == ECheckBoxState::Checked;
 }
 
 void SFoliageEditMeshDisplayItem::OnScaleYReapply(ECheckBoxState InState)
 {
+	FoliageSettingsPtr->Modify();
 	FoliageSettingsPtr->ReapplyScaleY = InState == ECheckBoxState::Checked;
 }
 
@@ -2219,26 +2236,26 @@ ECheckBoxState SFoliageEditMeshDisplayItem::IsScaleYReapplyChecked() const
 
 void SFoliageEditMeshDisplayItem::OnScaleYMinChanged(float InValue)
 {
-	FoliageSettingsPtr->ScaleMinY = InValue;
-
-	FoliageSettingsPtr->ScaleMaxY = FMath::Max(FoliageSettingsPtr->ScaleMaxY, InValue);
+	FoliageSettingsPtr->Modify();
+	FoliageSettingsPtr->ScaleY.Min = InValue;
+	FoliageSettingsPtr->ScaleY.Include(InValue);
 }
 
 float SFoliageEditMeshDisplayItem::GetScaleYMin() const
 {
-	return FoliageSettingsPtr->ScaleMinY;
+	return FoliageSettingsPtr->ScaleY.Min;
 }
 
 void SFoliageEditMeshDisplayItem::OnScaleYMaxChanged(float InValue)
 {
-	FoliageSettingsPtr->ScaleMaxY = InValue;
-
-	FoliageSettingsPtr->ScaleMinY = FMath::Min(FoliageSettingsPtr->ScaleMinY, InValue);
+	FoliageSettingsPtr->Modify();
+	FoliageSettingsPtr->ScaleY.Max = InValue;
+	FoliageSettingsPtr->ScaleY.Include(InValue);
 }
 
 float SFoliageEditMeshDisplayItem::GetScaleYMax() const
 {
-	return FoliageSettingsPtr->ScaleMaxY;
+	return FoliageSettingsPtr->ScaleY.Max;
 }
 
 ECheckBoxState SFoliageEditMeshDisplayItem::IsScaleYLockedChecked() const
@@ -2248,11 +2265,13 @@ ECheckBoxState SFoliageEditMeshDisplayItem::IsScaleYLockedChecked() const
 
 void SFoliageEditMeshDisplayItem::OnScaleYLocked(ECheckBoxState InState)
 {
+	FoliageSettingsPtr->Modify();
 	FoliageSettingsPtr->LockScaleY = InState == ECheckBoxState::Checked;
 }
 
 void SFoliageEditMeshDisplayItem::OnScaleZReapply(ECheckBoxState InState)
 {
+	FoliageSettingsPtr->Modify();
 	FoliageSettingsPtr->ReapplyScaleZ = InState == ECheckBoxState::Checked ? true : false;
 }
 
@@ -2263,26 +2282,26 @@ ECheckBoxState SFoliageEditMeshDisplayItem::IsScaleZReapplyChecked() const
 
 void SFoliageEditMeshDisplayItem::OnScaleZMinChanged(float InValue)
 {
-	FoliageSettingsPtr->ScaleMinZ = InValue;
-
-	FoliageSettingsPtr->ScaleMaxZ = FMath::Max(FoliageSettingsPtr->ScaleMaxZ, InValue);
+	FoliageSettingsPtr->Modify();
+	FoliageSettingsPtr->ScaleZ.Min = InValue;
+	FoliageSettingsPtr->ScaleZ.Include(InValue);
 }
 
 float SFoliageEditMeshDisplayItem::GetScaleZMin() const
 {
-	return FoliageSettingsPtr->ScaleMinZ;
+	return FoliageSettingsPtr->ScaleZ.Min;
 }
 
 void SFoliageEditMeshDisplayItem::OnScaleZMaxChanged(float InValue)
 {
-	FoliageSettingsPtr->ScaleMaxZ = InValue;
-
-	FoliageSettingsPtr->ScaleMinZ = FMath::Min(FoliageSettingsPtr->ScaleMinZ, InValue);
+	FoliageSettingsPtr->Modify();
+	FoliageSettingsPtr->ScaleZ.Max = InValue;
+	FoliageSettingsPtr->ScaleZ.Include(InValue);
 }
 
 float SFoliageEditMeshDisplayItem::GetScaleZMax() const
 {
-	return FoliageSettingsPtr->ScaleMaxZ;
+	return FoliageSettingsPtr->ScaleZ.Max;
 }
 
 ECheckBoxState SFoliageEditMeshDisplayItem::IsScaleZLockedChecked() const
@@ -2292,11 +2311,13 @@ ECheckBoxState SFoliageEditMeshDisplayItem::IsScaleZLockedChecked() const
 
 void SFoliageEditMeshDisplayItem::OnScaleZLocked(ECheckBoxState InState)
 {
+	FoliageSettingsPtr->Modify();
 	FoliageSettingsPtr->LockScaleZ = InState == ECheckBoxState::Checked;
 }
 
 void SFoliageEditMeshDisplayItem::OnZOffsetReapply(ECheckBoxState InState)
 {
+	FoliageSettingsPtr->Modify();
 	FoliageSettingsPtr->ReapplyZOffset = InState == ECheckBoxState::Checked ? true : false;
 }
 
@@ -2307,30 +2328,31 @@ ECheckBoxState SFoliageEditMeshDisplayItem::IsZOffsetReapplyChecked() const
 
 void SFoliageEditMeshDisplayItem::OnZOffsetMin(float InValue)
 {
-	FoliageSettingsPtr->ZOffsetMin = InValue;
-
-	FoliageSettingsPtr->ZOffsetMax = FMath::Max(FoliageSettingsPtr->ZOffsetMax, InValue);
+	FoliageSettingsPtr->Modify();
+	FoliageSettingsPtr->ZOffset.Min = InValue;
+	FoliageSettingsPtr->ZOffset.Include(InValue);
 }
 
 float SFoliageEditMeshDisplayItem::GetZOffsetMin() const
 {
-	return FoliageSettingsPtr->ZOffsetMin;
+	return FoliageSettingsPtr->ZOffset.Min;
 }
 
 void SFoliageEditMeshDisplayItem::OnZOffsetMax(float InValue)
 {
-	FoliageSettingsPtr->ZOffsetMax = InValue;
-
-	FoliageSettingsPtr->ZOffsetMin = FMath::Min(FoliageSettingsPtr->ZOffsetMin, InValue);
+	FoliageSettingsPtr->Modify();
+	FoliageSettingsPtr->ZOffset.Max = InValue;
+	FoliageSettingsPtr->ZOffset.Include(InValue);
 }
 
 float SFoliageEditMeshDisplayItem::GetZOffsetMax() const
 {
-	return FoliageSettingsPtr->ZOffsetMax;
+	return FoliageSettingsPtr->ZOffset.Max;
 }
 
 void SFoliageEditMeshDisplayItem::OnRandomPitchReapply(ECheckBoxState InState)
 {
+	FoliageSettingsPtr->Modify();
 	FoliageSettingsPtr->ReapplyRandomPitchAngle = InState == ECheckBoxState::Checked ? true : false;
 }
 
@@ -2341,6 +2363,7 @@ ECheckBoxState SFoliageEditMeshDisplayItem::IsRandomPitchReapplyChecked() const
 
 void SFoliageEditMeshDisplayItem::OnRandomPitchChanged(float InValue)
 {
+	FoliageSettingsPtr->Modify();
 	FoliageSettingsPtr->RandomPitchAngle = InValue;
 }
 
@@ -2351,6 +2374,7 @@ float SFoliageEditMeshDisplayItem::GetRandomPitch() const
 
 void SFoliageEditMeshDisplayItem::OnGroundSlopeReapply(ECheckBoxState InState)
 {
+	FoliageSettingsPtr->Modify();
 	FoliageSettingsPtr->ReapplyGroundSlope = InState == ECheckBoxState::Checked ? true : false;
 }
 
@@ -2361,6 +2385,7 @@ ECheckBoxState SFoliageEditMeshDisplayItem::IsGroundSlopeReapplyChecked() const
 
 void SFoliageEditMeshDisplayItem::OnGroundSlopeChanged(float InValue)
 {
+	FoliageSettingsPtr->Modify();
 	FoliageSettingsPtr->GroundSlope = InValue;
 }
 
@@ -2371,6 +2396,7 @@ float SFoliageEditMeshDisplayItem::GetGroundSlope() const
 
 void SFoliageEditMeshDisplayItem::OnHeightReapply(ECheckBoxState InState)
 {
+	FoliageSettingsPtr->Modify();
 	FoliageSettingsPtr->ReapplyHeight = InState == ECheckBoxState::Checked ? true : false;
 }
 
@@ -2381,30 +2407,31 @@ ECheckBoxState SFoliageEditMeshDisplayItem::IsHeightReapplyChecked() const
 
 void SFoliageEditMeshDisplayItem::OnHeightMinChanged(float InValue)
 {
-	FoliageSettingsPtr->HeightMin = InValue;
-
-	FoliageSettingsPtr->HeightMax = FMath::Max(FoliageSettingsPtr->HeightMax, InValue);
+	FoliageSettingsPtr->Modify();
+	FoliageSettingsPtr->Height.Min = InValue;
+	FoliageSettingsPtr->Height.Include(InValue);
 }
 
 float SFoliageEditMeshDisplayItem::GetHeightMin() const
 {
-	return FoliageSettingsPtr->HeightMin;
+	return FoliageSettingsPtr->Height.Min;
 }
 
 void SFoliageEditMeshDisplayItem::OnHeightMaxChanged(float InValue)
 {
-	FoliageSettingsPtr->HeightMax = InValue;
-
-	FoliageSettingsPtr->HeightMin = FMath::Min(FoliageSettingsPtr->HeightMin, InValue);
+	FoliageSettingsPtr->Modify();
+	FoliageSettingsPtr->Height.Max = InValue;
+	FoliageSettingsPtr->Height.Include(InValue);
 }
 
 float SFoliageEditMeshDisplayItem::GetHeightMax() const
 {
-	return FoliageSettingsPtr->HeightMax;
+	return FoliageSettingsPtr->Height.Max;
 }
 
 void SFoliageEditMeshDisplayItem::OnLandscapeLayerReapply(ECheckBoxState InState)
 {
+	FoliageSettingsPtr->Modify();
 	FoliageSettingsPtr->ReapplyLandscapeLayer = InState == ECheckBoxState::Checked ? true : false;
 }
 
@@ -2415,6 +2442,7 @@ ECheckBoxState SFoliageEditMeshDisplayItem::IsLandscapeLayerReapplyChecked() con
 
 void SFoliageEditMeshDisplayItem::OnLandscapeLayerChanged(const FText& InValue)
 {
+	FoliageSettingsPtr->Modify();
 	if (FoliageSettingsPtr->LandscapeLayers.Num() == 0)
 	{
 		FoliageSettingsPtr->LandscapeLayers.AddUninitialized(1);
@@ -2436,6 +2464,7 @@ FText SFoliageEditMeshDisplayItem::GetLandscapeLayer() const
 
 void SFoliageEditMeshDisplayItem::OnCollisionWithWorld(ECheckBoxState InState)
 {
+	FoliageSettingsPtr->Modify();
 	if (InState == ECheckBoxState::Checked)
 	{
 		FoliageSettingsPtr->CollisionWithWorld = true;
@@ -2458,6 +2487,7 @@ EVisibility SFoliageEditMeshDisplayItem::IsCollisionWithWorldVisible() const
 
 void SFoliageEditMeshDisplayItem::OnCollisionWithWorldReapply(ECheckBoxState InState)
 {
+	FoliageSettingsPtr->Modify();
 	FoliageSettingsPtr->ReapplyCollisionWithWorld = InState == ECheckBoxState::Checked ? true : false;
 }
 
@@ -2468,16 +2498,19 @@ ECheckBoxState SFoliageEditMeshDisplayItem::IsCollisionWithWorldReapplyChecked()
 
 void SFoliageEditMeshDisplayItem::OnCollisionScaleXChanged(float InValue)
 {
+	FoliageSettingsPtr->Modify();
 	FoliageSettingsPtr->CollisionScale.X = InValue;
 }
 
 void SFoliageEditMeshDisplayItem::OnCollisionScaleYChanged(float InValue)
 {
+	FoliageSettingsPtr->Modify();
 	FoliageSettingsPtr->CollisionScale.Y = InValue;
 }
 
 void SFoliageEditMeshDisplayItem::OnCollisionScaleZChanged(float InValue)
 {
+	FoliageSettingsPtr->Modify();
 	FoliageSettingsPtr->CollisionScale.Z = InValue;
 }
 
@@ -2498,6 +2531,7 @@ float SFoliageEditMeshDisplayItem::GetCollisionScaleZ() const
 
 void SFoliageEditMeshDisplayItem::OnVertexColorMask(ECheckBoxState InState, FoliageVertexColorMask Mask)
 {
+	FoliageSettingsPtr->Modify();
 	FoliageSettingsPtr->VertexColorMask = Mask;
 }
 
@@ -2508,6 +2542,7 @@ ECheckBoxState SFoliageEditMeshDisplayItem::IsVertexColorMaskChecked(FoliageVert
 
 void SFoliageEditMeshDisplayItem::OnVertexColorMaskThresholdChanged(float InValue)
 {
+	FoliageSettingsPtr->Modify();
 	FoliageSettingsPtr->VertexColorMaskThreshold = InValue;
 }
 
@@ -2523,6 +2558,7 @@ EVisibility SFoliageEditMeshDisplayItem::IsVertexColorMaskThresholdVisible() con
 
 void SFoliageEditMeshDisplayItem::OnVertexColorMaskInvert(ECheckBoxState InState)
 {
+	FoliageSettingsPtr->Modify();
 	FoliageSettingsPtr->VertexColorMaskInvert = (InState == ECheckBoxState::Checked);
 }
 
@@ -2533,6 +2569,7 @@ ECheckBoxState SFoliageEditMeshDisplayItem::IsVertexColorMaskInvertChecked() con
 
 void SFoliageEditMeshDisplayItem::OnVertexColorMaskReapply(ECheckBoxState InState)
 {
+	FoliageSettingsPtr->Modify();
 	FoliageSettingsPtr->ReapplyVertexColorMask = InState == ECheckBoxState::Checked ? true : false;
 }
 
@@ -2543,6 +2580,7 @@ ECheckBoxState SFoliageEditMeshDisplayItem::IsVertexColorMaskReapplyChecked() co
 
 FReply SFoliageEditMeshDisplayItem::OnMouseDownSelection(const FGeometry& MyGeometry, const FPointerEvent& MouseEvent)
 {
+	FoliageSettingsPtr->Modify();
 	FoliageSettingsPtr->IsSelected = !FoliageSettingsPtr->IsSelected;
 	return FReply::Handled();
 }
diff --git a/Engine/Source/Editor/FoliageEdit/Private/SFoliageEditMeshDisplayItem.h b/Engine/Source/Editor/FoliageEdit/Private/SFoliageEditMeshDisplayItem.h
--- a/Engine/Source/Editor/FoliageEdit/Private/SFoliageEditMeshDisplayItem.h
+++ b/Engine/Source/Editor/FoliageEdit/Private/SFoliageEditMeshDisplayItem.h
@@ -112,7 +112,7 @@ private:
 
 	/** Returns the tool-tip to use for the save settings button. */
 	FText GetSaveSettingsTooltip() const;
-	EVisibility IsSaveSettingsVisible() const;
+	bool IsSaveSettingsEnabled() const;
 
 	/** Determines what properties to display in the details panel. */
 	bool IsPropertyVisible(const FPropertyAndParent& PropertyAndParent) const;
diff --git a/Engine/Source/Editor/FoliageEdit/Private/SFoliageEditMeshRow.cpp b/Engine/Source/Editor/FoliageEdit/Private/SFoliageEditMeshRow.cpp
new file mode 100644
--- /dev/null
+++ b/Engine/Source/Editor/FoliageEdit/Private/SFoliageEditMeshRow.cpp
@@ -0,0 +1,179 @@
+// Copyright 1998-2015 Epic Games, Inc. All Rights Reserved.
+#include "UnrealEd.h"
+#include "FoliageType.h"
+#include "SFoliageEditMeshRow.h"
+
+#define LOCTEXT_NAMESPACE "FoliageEd_Mode"
+
+void SFoliageEditMeshRow::Construct(const FArguments& InArgs, TSharedRef<STableViewBase> InOwnerTableView)
+{
+	MeshInfo = InArgs._MeshInfo;
+	FoliageEditMode = InArgs._FoliageEditMode;
+	SMultiColumnTableRow<FFoliageMeshUIInfoPtr>::Construct(FSuperRowType::FArguments(), InOwnerTableView);
+}
+
+TSharedRef<SWidget> SFoliageEditMeshRow::GenerateWidgetForColumn(const FName& ColumnID)
+{
+	TSharedPtr<SWidget> TableRowContent = SNullWidget::NullWidget;
+
+	if (ColumnID == FoliageMeshColumns::ColumnID_ToggleMesh)
+	{
+		TableRowContent =
+			SNew(SCheckBox)
+			.OnCheckStateChanged( this, &SFoliageEditMeshRow::OnMeshCheckStateChanged)
+			.IsChecked( this, &SFoliageEditMeshRow::IsMeshChecked);
+	} 
+	else if (ColumnID == FoliageMeshColumns::ColumnID_MeshLabel)
+	{
+		TableRowContent =
+			SNew(SHorizontalBox)
+			
+			+SHorizontalBox::Slot()
+			.AutoWidth()
+			[
+				SNew(SExpanderArrow, SharedThis(this))
+			]
+			
+			+SHorizontalBox::Slot()
+			.VAlign(VAlign_Center)
+			.AutoWidth()
+			[
+				SNew(STextBlock)
+				.Text(this, &SFoliageEditMeshRow::GetMeshRowText)
+			];
+	}
+	else if (ColumnID == FoliageMeshColumns::ColumnID_InstanceCount)
+	{
+		TableRowContent =
+			SNew(SHorizontalBox)
+			
+			+SHorizontalBox::Slot()
+			.VAlign(VAlign_Center)
+			.AutoWidth()
+			[
+				SNew(STextBlock)
+				.Text(this, &SFoliageEditMeshRow::GetMeshInstanceCountText)
+			];
+	}
+	else if (ColumnID == FoliageMeshColumns::ColumnID_Save)
+	{
+		TableRowContent = 
+			SNew(SButton)
+			.ContentPadding(0)
+			.ButtonStyle(FEditorStyle::Get(), "ToggleButton")
+			.IsEnabled(this, &SFoliageEditMeshRow::IsSaveEnabled)
+			.OnClicked(this, &SFoliageEditMeshRow::OnSave)
+			.ToolTipText(LOCTEXT("SaveButtonToolTip", "Save foliage asset"))
+			.HAlign( HAlign_Center )
+			.VAlign( VAlign_Center )
+			.Content()
+			[
+				SNew(SImage)
+				.Image(this, &SFoliageEditMeshRow::GetMeshSaveBrush)
+			];
+	}
+
+	return TableRowContent.ToSharedRef();
+}
+
+FText SFoliageEditMeshRow::GetMeshRowText() const
+{
+	const UFoliageType* FoliageType = MeshInfo->Settings;
+	if (FoliageType->IsAsset())
+	{
+		return FText::FromString(FoliageType->GetName());
+	}
+	else
+	{
+		UStaticMesh* Mesh = FoliageType->GetStaticMesh();
+		if (Mesh)
+		{
+			return FText::FromString(Mesh->GetName());
+		}
+		else
+		{
+			return LOCTEXT("EmptyMeshName", "[Empty Mesh]");
+		}
+	}
+}
+	
+void SFoliageEditMeshRow::OnMeshCheckStateChanged(const ECheckBoxState NewCheckedState)
+{
+	bool bSelected = (NewCheckedState == ECheckBoxState::Checked);
+	if (MeshInfo->Settings->IsSelected != bSelected)
+	{
+		MeshInfo->Settings->Modify();
+		MeshInfo->Settings->IsSelected = bSelected;
+	}
+}
+
+ECheckBoxState SFoliageEditMeshRow::IsMeshChecked() const
+{
+	return MeshInfo->Settings->IsSelected ? ECheckBoxState::Checked : ECheckBoxState::Unchecked;
+}
+
+bool SFoliageEditMeshRow::IsSaveEnabled() const
+{
+	UFoliageType* FoliageType = MeshInfo->Settings;
+	// We want user to convert FoliageType to a shared asset
+	return (!FoliageType->IsAsset() || FoliageType->GetOutermost()->IsDirty());
+}
+
+FReply SFoliageEditMeshRow::OnSave()
+{
+	UFoliageType* SavedSettings = FoliageEditMode->SaveSettingsObject(MeshInfo->Settings);
+	if (SavedSettings)
+	{
+		MeshInfo->Settings = SavedSettings;
+	}
+
+	return FReply::Handled();
+}
+
+const FSlateBrush* SFoliageEditMeshRow::GetMeshSaveBrush() const
+{
+	return FEditorStyle::GetBrush("Level.SaveIcon16x");
+	
+	//if (LevelModel->IsLoaded())
+	//{
+	//	if (LevelModel->IsLocked())
+	//	{
+	//		return FEditorStyle::GetBrush( "Level.SaveDisabledIcon16x" );
+	//	}
+	//	else
+	//	{
+	//		if (LevelModel->IsDirty())
+	//		{
+	//			return SaveButton->IsHovered() ? FEditorStyle::GetBrush( "Level.SaveModifiedHighlightIcon16x" ) :
+	//												FEditorStyle::GetBrush( "Level.SaveModifiedIcon16x" );
+	//		}
+	//		else
+	//		{
+	//			return SaveButton->IsHovered() ? FEditorStyle::GetBrush( "Level.SaveHighlightIcon16x" ) :
+	//												FEditorStyle::GetBrush( "Level.SaveIcon16x" );
+	//		}
+	//	}								
+	//}
+	//else
+	//{
+	//	return FEditorStyle::GetBrush( "Level.EmptyIcon16x" );
+	//}	
+
+}
+
+FText SFoliageEditMeshRow::GetMeshInstanceCountText() const
+{
+	const int32	InstanceCountTotal = MeshInfo->InstanceCountTotal;
+	const int32	InstanceCountCurrentLevel = MeshInfo->InstanceCountCurrentLevel;
+	
+	if (InstanceCountCurrentLevel != InstanceCountTotal)
+	{
+		return FText::Format(LOCTEXT("InstanceCount", "{0} ({1})"), FText::AsNumber(InstanceCountCurrentLevel), FText::AsNumber(InstanceCountTotal));
+	}
+	else
+	{
+		return FText::AsNumber(InstanceCountCurrentLevel);
+	}
+}
+
+#undef LOCTEXT_NAMESPACE
diff --git a/Engine/Source/Editor/FoliageEdit/Private/SFoliageEditMeshRow.h b/Engine/Source/Editor/FoliageEdit/Private/SFoliageEditMeshRow.h
new file mode 100644
--- /dev/null
+++ b/Engine/Source/Editor/FoliageEdit/Private/SFoliageEditMeshRow.h
@@ -0,0 +1,42 @@
+// Copyright 1998-2015 Epic Games, Inc. All Rights Reserved.
+#pragma once
+
+namespace FoliageMeshColumns
+{
+	/** IDs for list columns */
+	static const FName ColumnID_MeshLabel("Mesh");
+	static const FName ColumnID_ToggleMesh("Toggle");
+	static const FName ColumnID_InstanceCount("InstanceCount");
+	static const FName ColumnID_Save("Save");
+};
+
+struct FFoliageMeshUIInfo;
+
+class SFoliageEditMeshRow
+	: public SMultiColumnTableRow<FFoliageMeshUIInfoPtr>
+{
+public:
+	SLATE_BEGIN_ARGS( SFoliageEditMeshRow )
+		{}
+		SLATE_ARGUMENT(FFoliageMeshUIInfoPtr,	MeshInfo)
+		SLATE_ARGUMENT(FEdModeFoliage*,			FoliageEditMode)
+	SLATE_END_ARGS()
+
+	void Construct(const FArguments& InArgs, TSharedRef<STableViewBase> OwnerTableView);
+	virtual TSharedRef<SWidget> GenerateWidgetForColumn(const FName& ColumnName) override;
+
+	FText GetMeshRowText() const;
+	
+	void OnMeshCheckStateChanged(const ECheckBoxState NewCheckedState);
+	ECheckBoxState IsMeshChecked() const;
+
+	bool IsSaveEnabled() const;
+	FReply OnSave();
+	const FSlateBrush* GetMeshSaveBrush() const;
+
+	FText GetMeshInstanceCountText() const;
+	
+private:
+	FFoliageMeshUIInfoPtr	MeshInfo;
+	FEdModeFoliage*			FoliageEditMode;
+};
\ No newline at end of file
diff --git a/Engine/Source/Runtime/Foliage/Private/InstancedFoliage.cpp b/Engine/Source/Runtime/Foliage/Private/InstancedFoliage.cpp
--- a/Engine/Source/Runtime/Foliage/Private/InstancedFoliage.cpp
+++ b/Engine/Source/Runtime/Foliage/Private/InstancedFoliage.cpp
@@ -42,6 +42,8 @@ struct FFoliageCustomVersion
 		ProceduralGuid = 4,
 		// Support for cross-level bases 
 		CrossLevelBase = 5,
+		// FoliageType for details customization
+		FoliageTypeCustomization = 6,
 		// -----<new versions can be added above this line>-------------------------------------------------
 		VersionPlusOne,
 		LatestVersion = VersionPlusOne - 1
@@ -252,19 +254,19 @@ UFoliageType::UFoliageType(const FObjectInitializer& ObjectInitializer)
 	AlignToNormal = true;
 	RandomYaw = true;
 	UniformScale = true;
-	ScaleMinX = 1.0f;
-	ScaleMinY = 1.0f;
-	ScaleMinZ = 1.0f;
-	ScaleMaxX = 1.0f;
-	ScaleMaxY = 1.0f;
-	ScaleMaxZ = 1.0f;
+	ScaleX.Min = 1.0f;
+	ScaleY.Min = 1.0f;
+	ScaleZ.Min = 1.0f;
+	ScaleX.Max = 1.0f;
+	ScaleY.Max = 1.0f;
+	ScaleZ.Max = 1.0f;
 	AlignMaxAngle = 0.0f;
 	RandomPitchAngle = 0.0f;
 	GroundSlope = 45.0f;
-	HeightMin = -262144.0f;
-	HeightMax = 262144.0f;
-	ZOffsetMin = 0.0f;
-	ZOffsetMax = 0.0f;
+	Height.Min = -262144.0f;
+	Height.Max = 262144.0f;
+	ZOffset.Min = 0.0f;
+	ZOffset.Max = 0.0f;
 	MinimumLayerWeight = 0.5f;
 	DisplayOrder = 0;
 	IsSelected = false;
@@ -313,17 +315,56 @@ UFoliageType::UFoliageType(const FObjectInitializer& ObjectInitializer)
 	Curve->AddKey(1.f, 1.f);
 
 	UpdateGuid = FGuid::NewGuid();
+
+	// Deprecated since FFoliageCustomVersion::FoliageTypeCustomization
+#if WITH_EDITORONLY_DATA
+	ScaleMinX_DEPRECATED = 1.0f;
+	ScaleMinY_DEPRECATED = 1.0f;
+	ScaleMinZ_DEPRECATED = 1.0f;
+	ScaleMaxX_DEPRECATED = 1.0f;
+	ScaleMaxY_DEPRECATED = 1.0f;
+	ScaleMaxZ_DEPRECATED = 1.0f;
+	HeightMin_DEPRECATED = -262144.0f;
+	HeightMax_DEPRECATED = 262144.0f;
+	ZOffsetMin_DEPRECATED = 0.0f;
+	ZOffsetMax_DEPRECATED = 0.0f;
+#endif// WITH_EDITORONLY_DATA
 }
 
 void UFoliageType::Serialize(FArchive& Ar)
 {
 	Super::Serialize(Ar);
 
+	Ar.UsingCustomVersion(FFoliageCustomVersion::GUID);
+
 	if (LandscapeLayer_DEPRECATED != NAME_None && LandscapeLayers.Num() == 0)	//we now store an array of names so initialize the array with the old name
 	{
 		LandscapeLayers.Add(LandscapeLayer_DEPRECATED);
 		LandscapeLayer_DEPRECATED = NAME_None;
 	}
+
+#if WITH_EDITORONLY_DATA
+	if (Ar.IsLoading() && Ar.CustomVer(FFoliageCustomVersion::GUID) < FFoliageCustomVersion::FoliageTypeCustomization)
+	{
+		ScaleX.Min = ScaleMinX_DEPRECATED;
+		ScaleX.Max = ScaleMaxX_DEPRECATED;
+
+		ScaleY.Min = ScaleMinY_DEPRECATED;
+		ScaleY.Max = ScaleMaxY_DEPRECATED;
+
+		ScaleZ.Min = ScaleMinZ_DEPRECATED;
+		ScaleZ.Max = ScaleMaxZ_DEPRECATED;
+
+		Height.Min = HeightMin_DEPRECATED;
+		Height.Max = HeightMax_DEPRECATED;
+
+		ZOffset.Min = ZOffsetMin_DEPRECATED;
+		ZOffset.Max = ZOffsetMax_DEPRECATED;
+
+		CullDistance.Min = StartCullDistance_DEPRECATED;
+		CullDistance.Max = EndCullDistance_DEPRECATED;
+	}
+#endif// WITH_EDITORONLY_DATA
 }
 
 
@@ -483,8 +524,8 @@ void FFoliageMeshInfo::UpdateComponentSettings(const UFoliageType* InSettings)
 	if (Component)
 	{
 		Component->Mobility = InSettings->bEnableStaticLighting ? EComponentMobility::Static : EComponentMobility::Movable;
-		Component->InstanceStartCullDistance = InSettings->StartCullDistance;
-		Component->InstanceEndCullDistance = InSettings->EndCullDistance;
+		Component->InstanceStartCullDistance = InSettings->CullDistance.Min;
+		Component->InstanceEndCullDistance = InSettings->CullDistance.Max;
 
 		Component->CastShadow = InSettings->CastShadow;
 		Component->bCastDynamicShadow = InSettings->bCastDynamicShadow;
@@ -2109,15 +2150,15 @@ bool FPotentialInstance::PlaceInstance(const UWorld* InWorld, const UFoliageType
 	{
 		if (Settings->UniformScale)
 		{
-			float Scale = Settings->ScaleMinX + FMath::FRand() * (Settings->ScaleMaxX - Settings->ScaleMinX);
+			float Scale = Settings->ScaleX.Interpolate(FMath::FRand());
 			Inst.DrawScale3D = FVector(Scale, Scale, Scale);
 		}
 		else
 		{
 			float LockRand = FMath::FRand();
-			Inst.DrawScale3D.X = Settings->ScaleMinX + (Settings->LockScaleX ? LockRand : FMath::FRand()) * (Settings->ScaleMaxX - Settings->ScaleMinX);
-			Inst.DrawScale3D.Y = Settings->ScaleMinY + (Settings->LockScaleY ? LockRand : FMath::FRand()) * (Settings->ScaleMaxY - Settings->ScaleMinY);
-			Inst.DrawScale3D.Z = Settings->ScaleMinZ + (Settings->LockScaleZ ? LockRand : FMath::FRand()) * (Settings->ScaleMaxZ - Settings->ScaleMinZ);
+			Inst.DrawScale3D.X = Settings->ScaleX.Interpolate(Settings->LockScaleX ? LockRand : FMath::FRand());
+			Inst.DrawScale3D.Y = Settings->ScaleY.Interpolate(Settings->LockScaleY ? LockRand : FMath::FRand());
+			Inst.DrawScale3D.Z = Settings->ScaleZ.Interpolate(Settings->LockScaleZ ? LockRand : FMath::FRand());
 		}
 	}
 	else
@@ -2126,8 +2167,7 @@ bool FPotentialInstance::PlaceInstance(const UWorld* InWorld, const UFoliageType
 		Inst.DrawScale3D = FVector(Settings->GetScaleForAge(DesiredInstance.Age));
 	}
 
-
-	Inst.ZOffset = Settings->ZOffsetMin + FMath::FRand() * (Settings->ZOffsetMax - Settings->ZOffsetMin);
+	Inst.ZOffset = Settings->ZOffset.Interpolate(FMath::FRand());
 
 	Inst.Location = HitLocation;
 
diff --git a/Engine/Source/Runtime/Foliage/Public/FoliageType.h b/Engine/Source/Runtime/Foliage/Public/FoliageType.h
--- a/Engine/Source/Runtime/Foliage/Public/FoliageType.h
+++ b/Engine/Source/Runtime/Foliage/Public/FoliageType.h
@@ -32,22 +32,16 @@ class UFoliageType : public UObject
 	float Radius;
 
 	UPROPERTY(EditAnywhere, Category=Painting)
-	float ScaleMinX;
+	FFloatInterval ScaleX;
 
 	UPROPERTY(EditAnywhere, Category=Painting)
-	float ScaleMinY;
+	FFloatInterval ScaleY;
 
 	UPROPERTY(EditAnywhere, Category=Painting)
-	float ScaleMinZ;
+	FFloatInterval ScaleZ;
 
 	UPROPERTY(EditAnywhere, Category=Painting)
-	float ScaleMaxX;
-
-	UPROPERTY(EditAnywhere, Category=Painting)
-	float ScaleMaxY;
-
-	UPROPERTY(EditAnywhere, Category=Painting)
-	float ScaleMaxZ;
+	uint32 UniformScale:1;
 
 	UPROPERTY(EditAnywhere, Category=Painting)
 	uint32 LockScaleX:1;
@@ -58,48 +52,39 @@ class UFoliageType : public UObject
 	UPROPERTY(EditAnywhere, Category=Painting)
 	uint32 LockScaleZ:1;
 
-	UPROPERTY(EditAnywhere, Category=General, meta=(Subcategory = "Placement"))
+	UPROPERTY(EditAnywhere, Category=Painting)
 	float AlignMaxAngle;
 
-	UPROPERTY(EditAnywhere, Category = General, meta = (Subcategory = "Placement"))
+	UPROPERTY(EditAnywhere, Category=Painting)
 	float RandomPitchAngle;
 
-	UPROPERTY(EditAnywhere, Category = General, meta = (Subcategory = "Placement"))
+	UPROPERTY(EditAnywhere, Category=Painting)
 	float GroundSlope;
 
-	UPROPERTY(EditAnywhere, Category = General, meta = (Subcategory = "Placement"))
+	UPROPERTY(EditAnywhere, Category=Painting)
 	float MinGroundSlope;
 
-	UPROPERTY(EditAnywhere, Category = General, meta = (Subcategory = "Placement"))
-	float HeightMin;
-
-	UPROPERTY(EditAnywhere, Category = General, meta = (Subcategory = "Placement"))
-	float HeightMax;
+	UPROPERTY(EditAnywhere, Category=Painting)
+	FFloatInterval Height;
 
-	UPROPERTY(EditAnywhere, Category = General, meta = (Subcategory = "Placement"))
+	UPROPERTY(EditAnywhere, Category=Painting)
 	TArray<FName> LandscapeLayers;
 
 	UPROPERTY()
 	FName LandscapeLayer_DEPRECATED;
 	
-	UPROPERTY(EditAnywhere, Category = General, meta = (Subcategory = "Placement"))
+	UPROPERTY(EditAnywhere, Category=Painting)
 	float MinimumLayerWeight;
 
-	UPROPERTY(EditAnywhere, Category = General, meta = (Subcategory = "Placement"))
-	uint32 AlignToNormal:1;
-
-	UPROPERTY(EditAnywhere, Category = General, meta = (Subcategory = "Placement"))
-	uint32 RandomYaw:1;
-
 	UPROPERTY(EditAnywhere, Category=Painting)
-	uint32 UniformScale:1;
+	uint32 AlignToNormal:1;
 
 	UPROPERTY(EditAnywhere, Category=Painting)
-	float ZOffsetMin;
+	uint32 RandomYaw:1;
 
 	UPROPERTY(EditAnywhere, Category=Painting)
-	float ZOffsetMax;
-
+	FFloatInterval ZOffset;
+	
 	UPROPERTY(EditAnywhere, Category=Painting)
 	uint32 CollisionWithWorld:1;
 
@@ -167,17 +152,13 @@ class UFoliageType : public UObject
 	UPROPERTY()
 	uint32 ReapplyVertexColorMask:1;
 
-	/* The distance where instances will begin to fade out if using a PerInstanceFadeAmount material node. 0 disables. */
-	UPROPERTY(EditAnywhere, Category = General, meta = (Subcategory = "Culling"))
-	int32 StartCullDistance;
-
 	/**
-	 * The distance where instances will have completely faded out when using a PerInstanceFadeAmount material node. 0 disables. 
+	 * The distance where instances will begin to fade out if using a PerInstanceFadeAmount material node. 0 disables.
 	 * When the entire cluster is beyond this distance, the cluster is completely culled and not rendered at all.
 	 */
-	UPROPERTY(EditAnywhere, Category = General, meta = (Subcategory = "Culling"))
-	int32 EndCullDistance;
-
+	UPROPERTY(EditAnywhere, Category = General)
+	FInt32Interval CullDistance;
+	
 	UPROPERTY()
 	int32 DisplayOrder;
 
@@ -234,7 +215,7 @@ class UFoliageType : public UObject
 	int32 OverriddenLightMapRes;
 
 	/** Custom collision for foliage */
-	UPROPERTY(EditAnywhere, Category = Painting, meta = (Subcategory = "Collision", HideObjectType = true))
+	UPROPERTY(EditAnywhere, Category = General, meta = (HideObjectType = true))
 	struct FBodyInstance BodyInstance;
 
 	/* Gets/Sets the mesh associated with this FoliageType */
@@ -327,6 +308,45 @@ private:
 	/** The curve used to interpolate the instance scale.*/
 	UPROPERTY(Category = Procedural, EditAnywhere, meta = (Subcategory = "Growth"))
 	FRuntimeFloatCurve ScaleCurve;
+	
+	// Deprecated since FFoliageCustomVersion::FoliageTypeCustomization
+#if WITH_EDITORONLY_DATA
+	UPROPERTY()
+	float ScaleMinX_DEPRECATED;
+
+	UPROPERTY()
+	float ScaleMinY_DEPRECATED;
+
+	UPROPERTY()
+	float ScaleMinZ_DEPRECATED;
+
+	UPROPERTY()
+	float ScaleMaxX_DEPRECATED;
+
+	UPROPERTY()
+	float ScaleMaxY_DEPRECATED;
+
+	UPROPERTY()
+	float ScaleMaxZ_DEPRECATED;
+
+	UPROPERTY()
+	float HeightMin_DEPRECATED;
+
+	UPROPERTY()
+	float HeightMax_DEPRECATED;
+
+	UPROPERTY()
+	float ZOffsetMin_DEPRECATED;
+
+	UPROPERTY()
+	float ZOffsetMax_DEPRECATED;
+
+	UPROPERTY()
+	int32 StartCullDistance_DEPRECATED;
+	
+	UPROPERTY()
+	int32 EndCullDistance_DEPRECATED;
+#endif// WITH_EDITORONLY_DATA
 };
 
 
