fixing non-unity build

diff --git a/Engine/Source/Editor/FoliageEdit/Private/SFoliageEdit2.cpp b/Engine/Source/Editor/FoliageEdit/Private/SFoliageEdit2.cpp
--- a/Engine/Source/Editor/FoliageEdit/Private/SFoliageEdit2.cpp
+++ b/Engine/Source/Editor/FoliageEdit/Private/SFoliageEdit2.cpp
@@ -4,17 +4,17 @@
 
 #include "SScaleBox.h"
 
+#include "FoliageType.h"
+#include "FoliageEdMode.h"
 #include "SFoliageEdit2.h"
 #include "FoliageEditActions.h"
-#include "FoliageEdMode.h"
-#include "FoliageType.h"
 #include "Editor/UnrealEd/Public/AssetThumbnail.h"
 #include "Editor/UnrealEd/Public/DragAndDrop/AssetDragDropOp.h"
 #include "Editor/ContentBrowser/Public/ContentBrowserModule.h"
 #include "Editor/UnrealEd/Public/AssetSelection.h"
 #include "Editor/IntroTutorials/Public/IIntroTutorials.h"
 #include "Editor/PropertyEditor/Public/IDetailsView.h"
-#include "SFoliageEditMeshDisplayItem.h"
+#include "Editor/PropertyEditor/Public/PropertyEditorModule.h"
 #include "Engine/StaticMesh.h"
 #include "SFoliageEditMeshRow.h"
 
diff --git a/Engine/Source/Editor/FoliageEdit/Private/SFoliageEdit2.h b/Engine/Source/Editor/FoliageEdit/Private/SFoliageEdit2.h
--- a/Engine/Source/Editor/FoliageEdit/Private/SFoliageEdit2.h
+++ b/Engine/Source/Editor/FoliageEdit/Private/SFoliageEdit2.h
@@ -9,6 +9,7 @@ struct FFoliageMeshUIInfo;
 class FAssetThumbnailPool;
 class IDetailsView;
 
+typedef TSharedPtr<FFoliageMeshUIInfo> FFoliageMeshUIInfoPtr; //should match typedef in FoliageEdMode.h
 typedef STreeView<FFoliageMeshUIInfoPtr> SFoliageMeshTree;
 
 class SFoliageEdit2 : public SCompoundWidget
diff --git a/Engine/Source/Editor/FoliageEdit/Private/SFoliageEditMeshRow.cpp b/Engine/Source/Editor/FoliageEdit/Private/SFoliageEditMeshRow.cpp
--- a/Engine/Source/Editor/FoliageEdit/Private/SFoliageEditMeshRow.cpp
+++ b/Engine/Source/Editor/FoliageEdit/Private/SFoliageEditMeshRow.cpp
@@ -1,6 +1,7 @@
 // Copyright 1998-2015 Epic Games, Inc. All Rights Reserved.
 #include "UnrealEd.h"
 #include "FoliageType.h"
+#include "FoliageEdMode.h"
 #include "SFoliageEditMeshRow.h"
 
 #define LOCTEXT_NAMESPACE "FoliageEd_Mode"
diff --git a/Engine/Source/Editor/FoliageEdit/Private/SFoliageEditMeshRow.h b/Engine/Source/Editor/FoliageEdit/Private/SFoliageEditMeshRow.h
--- a/Engine/Source/Editor/FoliageEdit/Private/SFoliageEditMeshRow.h
+++ b/Engine/Source/Editor/FoliageEdit/Private/SFoliageEditMeshRow.h
@@ -10,7 +10,9 @@ namespace FoliageMeshColumns
 	static const FName ColumnID_Save("Save");
 };
 
+class FEdModeFoliage;
 struct FFoliageMeshUIInfo;
+typedef TSharedPtr<FFoliageMeshUIInfo> FFoliageMeshUIInfoPtr; //should match typedef in FoliageEdMode.h
 
 class SFoliageEditMeshRow
 	: public SMultiColumnTableRow<FFoliageMeshUIInfoPtr>
