New foliage tool UI
UE-8288

diff --git a/Engine/Source/Editor/DetailCustomizations/Private/IntervalStructCustomization.cpp b/Engine/Source/Editor/DetailCustomizations/Private/IntervalStructCustomization.cpp
--- a/Engine/Source/Editor/DetailCustomizations/Private/IntervalStructCustomization.cpp
+++ b/Engine/Source/Editor/DetailCustomizations/Private/IntervalStructCustomization.cpp
@@ -162,7 +162,7 @@ TOptional<NumericType> FIntervalStructCustomization<NumericType>::OnGetValue(TWe
 	auto ValueSharedPtr = ValueWeakPtr.Pin();
 
 	NumericType Value;
-	if (ensure(ValueSharedPtr->GetValue(Value) == FPropertyAccess::Success))
+	if (ValueSharedPtr->GetValue(Value) == FPropertyAccess::Success)
 	{
 		return TOptional<NumericType>(Value);
 	}
diff --git a/Engine/Source/Editor/FoliageEdit/Private/FoliageEdMode.cpp b/Engine/Source/Editor/FoliageEdit/Private/FoliageEdMode.cpp
--- a/Engine/Source/Editor/FoliageEdit/Private/FoliageEdMode.cpp
+++ b/Engine/Source/Editor/FoliageEdit/Private/FoliageEdMode.cpp
@@ -354,8 +354,8 @@ void FEdModeFoliage::FoliageBrushTrace(FEditorViewportClient* ViewportClient, in
 					(UISettings.bFilterLandscape || !PrimComp->IsA(ULandscapeHeightfieldCollisionComponent::StaticClass())) &&
 					(UISettings.bFilterStaticMesh || !PrimComp->IsA(UStaticMeshComponent::StaticClass())) &&
 					(UISettings.bFilterBSP || !PrimComp->IsA(UModelComponent::StaticClass())) &&
-					(UISettings.bFilterTranslucent || !Material || !IsTranslucentBlendMode(Material->GetBlendMode()))
-					)
+					(UISettings.bFilterTranslucent || !Material || !IsTranslucentBlendMode(Material->GetBlendMode())) &&
+					(CanPaint(PrimComp->GetComponentLevel())))
 				{
 					// Adjust the sphere brush
 					BrushLocation = Hit.Location;
@@ -426,8 +426,8 @@ static bool CheckLocationForPotentialInstance_ThreadSafe(const UFoliageType* Set
 	}
 
 	// Check slope
-	const float MaxNormalAngle = FMath::Cos(FMath::DegreesToRadians(Settings->GroundSlope));
-	const float MinNormalAngle = FMath::Cos(FMath::DegreesToRadians(Settings->MinGroundSlope));
+	const float MaxNormalAngle = FMath::Cos(FMath::DegreesToRadians(Settings->GroundSlopeAngle.Max));
+	const float MinNormalAngle = FMath::Cos(FMath::DegreesToRadians(Settings->GroundSlopeAngle.Min));
 	if (MaxNormalAngle > Normal.Z || MinNormalAngle < Normal.Z)	//keep in mind Hit.ImpactNormal.Z is (0,0,1) dot normal. However, ground slope is with relation to plane vector, not plane normal - so we swap comparisons
 	{
 		return false;
@@ -1276,55 +1276,34 @@ void FEdModeFoliage::ReapplyInstancesForBrush(UWorld* InWorld, AInstancedFoliage
 			}
 
 			// Reapply scale
-			if (Settings->UniformScale)
-			{
-				if (Settings->ReapplyScaleX)
-				{
-					float Scale = Settings->ScaleX.Interpolate(FMath::FRand());
-					Instance.DrawScale3D = FVector(Scale, Scale, Scale);
-					bUpdated = true;
-				}
-			}
-			else
+			FVector NewScale = Settings->GetRandomScale();
+			
+			if (Settings->ReapplyScaleX)
 			{
-				float LockRand;
-				// If we're doing axis scale locking, get an existing scale for a locked axis that we're not changing, for use as the locked scale value.
-				if (Settings->LockScaleX && !Settings->ReapplyScaleX && Settings->ScaleX.Size() > KINDA_SMALL_NUMBER)
-				{
-					LockRand = (Instance.DrawScale3D.X - Settings->ScaleX.Min) / Settings->ScaleX.Size();
-				}
-				else if (Settings->LockScaleY && !Settings->ReapplyScaleY && Settings->ScaleY.Size() > KINDA_SMALL_NUMBER)
+				if (Settings->Scaling == EFoliageScaling::Uniform)
 				{
-					LockRand = (Instance.DrawScale3D.Y - Settings->ScaleY.Min) / Settings->ScaleY.Size();
-				}
-				else if (Settings->LockScaleZ && !Settings->ReapplyScaleZ && Settings->ScaleZ.Size() > KINDA_SMALL_NUMBER)
-				{
-					LockRand = (Instance.DrawScale3D.Z - Settings->ScaleZ.Min) / Settings->ScaleZ.Size();
+					Instance.DrawScale3D = NewScale;
 				}
 				else
 				{
-					LockRand = FMath::FRand();
-				}
-
-				if (Settings->ReapplyScaleX)
-				{
-					Instance.DrawScale3D.X = Settings->ScaleX.Interpolate(Settings->LockScaleX ? LockRand : FMath::FRand());
-					bUpdated = true;
+					Instance.DrawScale3D.X = NewScale.X;
 				}
+				bUpdated = true;
+			}
 
-				if (Settings->ReapplyScaleY)
-				{
-					Instance.DrawScale3D.Y = Settings->ScaleY.Interpolate(Settings->LockScaleY ? LockRand : FMath::FRand());
-					bUpdated = true;
-				}
+			if (Settings->ReapplyScaleY)
+			{
+				Instance.DrawScale3D.Y = NewScale.Y;
+				bUpdated = true;
+			}
 
-				if (Settings->ReapplyScaleZ)
-				{
-					Instance.DrawScale3D.Z = Settings->ScaleZ.Interpolate(Settings->LockScaleZ ? LockRand : FMath::FRand());
-					bUpdated = true;
-				}
+			if (Settings->ReapplyScaleY)
+			{
+				Instance.DrawScale3D.Z = NewScale.Z;
+				bUpdated = true;
 			}
 
+			// Reapply ZOffset
 			if (Settings->ReapplyZOffset)
 			{
 				Instance.ZOffset = Settings->ZOffset.Interpolate(FMath::FRand());
@@ -1353,8 +1332,9 @@ void FEdModeFoliage::ReapplyInstancesForBrush(UWorld* InWorld, AInstancedFoliage
 					// Cull instances that don't meet the ground slope check.
 					if (Settings->ReapplyGroundSlope)
 					{
-						if ((Settings->GroundSlope > 0.f && Hit.Normal.Z <= FMath::Cos(PI * Settings->GroundSlope / 180.f) - SMALL_NUMBER) ||
-							(Settings->GroundSlope < 0.f && Hit.Normal.Z >= FMath::Cos(-PI * Settings->GroundSlope / 180.f) + SMALL_NUMBER))
+						const float MaxNormalAngle = FMath::Cos(FMath::DegreesToRadians(Settings->GroundSlopeAngle.Max));
+						const float MinNormalAngle = FMath::Cos(FMath::DegreesToRadians(Settings->GroundSlopeAngle.Min));
+						if (MaxNormalAngle > Hit.Normal.Z || MinNormalAngle < Hit.Normal.Z)
 						{
 							InstancesToDelete.Add(InstanceIndex);
 							if (bReapplyLocation)
@@ -1809,8 +1789,9 @@ void FEdModeFoliage::ApplyPaintBucket_Add(AActor* Actor)
 				FFoliagePaintBucketTriangle& Triangle = PotentialTriangles[TriIdx];
 
 				// Check if we can reject this triangle based on normal.
-				if ((Settings->GroundSlope > 0.f && Triangle.WorldNormal.Z <= FMath::Cos(PI * Settings->GroundSlope / 180.f) - SMALL_NUMBER) ||
-					(Settings->GroundSlope < 0.f && Triangle.WorldNormal.Z >= FMath::Cos(-PI * Settings->GroundSlope / 180.f) + SMALL_NUMBER))
+				const float MaxNormalAngle = FMath::Cos(FMath::DegreesToRadians(Settings->GroundSlopeAngle.Max));
+				const float MinNormalAngle = FMath::Cos(FMath::DegreesToRadians(Settings->GroundSlopeAngle.Min));
+				if (Triangle.WorldNormal.Z < MaxNormalAngle || Triangle.WorldNormal.Z > MinNormalAngle)
 				{
 					continue;
 				}
@@ -2127,6 +2108,19 @@ void FEdModeFoliage::OnInstanceCountUpdated(const UFoliageType* FoliageType)
 	FoliageMeshList[EntryIndex]->InstanceCountCurrentLevel = InstanceCountCurrentLevel;
 }
 
+bool FEdModeFoliage::CanPaint(const ULevel* InLevel)
+{
+	for (const auto& MeshUIPtr : FoliageMeshList)
+	{
+		if (MeshUIPtr->Settings->IsSelected && CanPaint(MeshUIPtr->Settings, InLevel))
+		{
+			return true;
+		}
+	}
+
+	return false;
+}
+
 bool FEdModeFoliage::CanPaint(const UFoliageType* FoliageType, const ULevel* InLevel)
 {
 	// Non-shared objects can be painted only into their own level
diff --git a/Engine/Source/Editor/FoliageEdit/Private/FoliageEdMode.h b/Engine/Source/Editor/FoliageEdit/Private/FoliageEdMode.h
--- a/Engine/Source/Editor/FoliageEdit/Private/FoliageEdMode.h
+++ b/Engine/Source/Editor/FoliageEdit/Private/FoliageEdMode.h
@@ -282,6 +282,9 @@ public:
 
 	/**  */
 	void OnInstanceCountUpdated(const UFoliageType* FoliageType);
+
+	/** Whether any of the selected foliage types can be painted into level */
+	bool CanPaint(const ULevel* InLevel);
 		
 	/** Whether specified FoliageType can be painted into level */
 	static bool CanPaint(const UFoliageType* FoliageType, const ULevel* InLevel);
diff --git a/Engine/Source/Editor/FoliageEdit/Private/FoliageEdModeToolkit.cpp b/Engine/Source/Editor/FoliageEdit/Private/FoliageEdModeToolkit.cpp
--- a/Engine/Source/Editor/FoliageEdit/Private/FoliageEdModeToolkit.cpp
+++ b/Engine/Source/Editor/FoliageEdit/Private/FoliageEdModeToolkit.cpp
@@ -3,16 +3,9 @@
 #include "UnrealEd.h"
 #include "FoliageEdModeToolkit.h"
 #include "SFoliageEdit.h"
-#include "SFoliageEdit2.h"
 
 #define LOCTEXT_NAMESPACE "FoliageEditMode"
 
-static TAutoConsoleVariable<int32> CVarFoliageNewUI(
-	TEXT("foliage.NewUI"),
-	0,
-	TEXT("Enables new foliage UI"));
-
-
 void FFoliageEdModeToolkit::RegisterTabSpawners(const TSharedRef<class FTabManager>& TabManager)
 {
 
@@ -25,17 +18,7 @@ void FFoliageEdModeToolkit::UnregisterTabSpawners(const TSharedRef<class FTabMan
 
 void FFoliageEdModeToolkit::Init(const TSharedPtr< class IToolkitHost >& InitToolkitHost)
 {
-	bool NewUI = CVarFoliageNewUI.GetValueOnAnyThread() != 0;
-	
-	if (NewUI)
-	{
-		FoliageEdWidget2 = SNew(SFoliageEdit2);
-	}
-	else
-	{
-		FoliageEdWidget = SNew(SFoliageEdit);
-	}
-	
+	FoliageEdWidget = SNew(SFoliageEdit);
 
 	FModeToolkit::Init(InitToolkitHost);
 }
@@ -57,26 +40,12 @@ class FEdMode* FFoliageEdModeToolkit::GetEditorMode() const
 
 TSharedPtr<SWidget> FFoliageEdModeToolkit::GetInlineContent() const
 {
-	if (FoliageEdWidget2.IsValid())
-	{
-		return FoliageEdWidget2;
-	}
-	else
-	{
-		return FoliageEdWidget;
-	}
+	return FoliageEdWidget;
 }
 
 void FFoliageEdModeToolkit::RefreshFullList()
 {
-	if (FoliageEdWidget2.IsValid())
-	{
-		FoliageEdWidget2->RefreshFullList();
-	}
-	else
-	{
-		FoliageEdWidget->RefreshFullList();
-	}
+	FoliageEdWidget->RefreshFullList();
 }
 
 #undef LOCTEXT_NAMESPACE
diff --git a/Engine/Source/Editor/FoliageEdit/Private/FoliageEdModeToolkit.h b/Engine/Source/Editor/FoliageEdit/Private/FoliageEdModeToolkit.h
--- a/Engine/Source/Editor/FoliageEdit/Private/FoliageEdModeToolkit.h
+++ b/Engine/Source/Editor/FoliageEdit/Private/FoliageEdModeToolkit.h
@@ -26,5 +26,4 @@ public:
 
 private:
 	TSharedPtr< class SFoliageEdit > FoliageEdWidget;
-	TSharedPtr< class SFoliageEdit2 > FoliageEdWidget2;
 };
\ No newline at end of file
diff --git a/Engine/Source/Editor/FoliageEdit/Private/FoliageTypeDetails.cpp b/Engine/Source/Editor/FoliageEdit/Private/FoliageTypeDetails.cpp
--- a/Engine/Source/Editor/FoliageEdit/Private/FoliageTypeDetails.cpp
+++ b/Engine/Source/Editor/FoliageEdit/Private/FoliageTypeDetails.cpp
@@ -11,6 +11,17 @@ TSharedRef<IDetailCustomization> FFoliageTypeDetails::MakeInstance()
 
 }
 
+static void HideFoliageCategory(IDetailLayoutBuilder& DetailBuilder, FName Category)
+{
+	TArray<TSharedRef<IPropertyHandle>> CategoryProperties;
+	DetailBuilder.EditCategory(Category).GetDefaultProperties(CategoryProperties, true, true);
+
+	for (auto& PropertyHandle : CategoryProperties)
+	{
+		DetailBuilder.HideProperty(PropertyHandle);
+	}
+}
+
 void AddSubcategoryProperties(IDetailCategoryBuilder& CategoryBuilder)
 {
 	TArray<TSharedRef<IPropertyHandle>> CategoryProperties;
@@ -47,7 +58,13 @@ void AddSubcategoryProperties(IDetailCategoryBuilder& CategoryBuilder)
 
 void FFoliageTypeDetails::CustomizeDetails(IDetailLayoutBuilder& DetailBuilder)
 {
-	AddSubcategoryProperties(DetailBuilder.EditCategory("General"));
-	AddSubcategoryProperties(DetailBuilder.EditCategory("Procedural"));
-	AddSubcategoryProperties(DetailBuilder.EditCategory("Painting"));
+	const static FName PaintingName("Painting");
+	const static FName PlacementName("Placement");
+	const static FName ProceduralName("Procedural");
+	const static FName InstanceSettingsName("InstanceSettings");
+	
+	HideFoliageCategory(DetailBuilder, PaintingName);
+	AddSubcategoryProperties(DetailBuilder.EditCategory(PlacementName));
+	AddSubcategoryProperties(DetailBuilder.EditCategory(ProceduralName));
+	AddSubcategoryProperties(DetailBuilder.EditCategory(InstanceSettingsName));
 }
\ No newline at end of file
diff --git a/Engine/Source/Editor/FoliageEdit/Private/FoliageTypePaintingCustomization.cpp b/Engine/Source/Editor/FoliageEdit/Private/FoliageTypePaintingCustomization.cpp
new file mode 100644
--- /dev/null
+++ b/Engine/Source/Editor/FoliageEdit/Private/FoliageTypePaintingCustomization.cpp
@@ -0,0 +1,317 @@
+// Copyright 1998-2015 Epic Games, Inc. All Rights Reserved.
+#include "UnrealEd.h"
+#include "FoliageEditModule.h"
+#include "PropertyEditing.h"
+#include "FoliageType.h"
+#include "FoliageEdMode.h"
+#include "FoliageTypePaintingCustomization.h"
+
+#define LOCTEXT_NAMESPACE "FoliageEd_Mode"
+
+/////////////////////////////////////////////////////
+// FFoliageTypePaintingCustomization 
+TSharedRef<IDetailCustomization> FFoliageTypePaintingCustomization::MakeInstance(FEdModeFoliage* InFoliageEditMode)
+{
+	TSharedRef<FFoliageTypePaintingCustomization> Instance = MakeShareable(new FFoliageTypePaintingCustomization(InFoliageEditMode));
+	return Instance;
+}
+
+FFoliageTypePaintingCustomization::FFoliageTypePaintingCustomization(FEdModeFoliage* InFoliageEditMode)
+{
+	FoliageEditMode = InFoliageEditMode;
+}
+
+void FFoliageTypePaintingCustomization::CustomizeDetails(IDetailLayoutBuilder& DetailLayoutBuilder)
+{
+	// Hide categories we are not going to customize
+	const static FName PaintingName("Painting");
+	const static FName PlacementName("Placement");
+	const static FName ProceduralName("Procedural");
+	HideFoliageCategory(DetailLayoutBuilder, PaintingName);
+	HideFoliageCategory(DetailLayoutBuilder, PlacementName);
+	HideFoliageCategory(DetailLayoutBuilder, ProceduralName);
+	
+	//
+	// Re-add customized properties 
+	//
+		
+	// Get the properties we need
+	TSharedPtr<IPropertyHandle> InvalidProperty;
+	IDetailCategoryBuilder& PaintingCategory = DetailLayoutBuilder.EditCategory("Painting");
+	
+	Density = DetailLayoutBuilder.GetProperty(GET_MEMBER_NAME_CHECKED(UFoliageType, Density));
+	Radius = DetailLayoutBuilder.GetProperty(GET_MEMBER_NAME_CHECKED(UFoliageType, Radius));
+	AlignToNormal = DetailLayoutBuilder.GetProperty(GET_MEMBER_NAME_CHECKED(UFoliageType, AlignToNormal));
+	AlignMaxAngle = DetailLayoutBuilder.GetProperty(GET_MEMBER_NAME_CHECKED(UFoliageType, AlignMaxAngle));
+	RandomYaw = DetailLayoutBuilder.GetProperty(GET_MEMBER_NAME_CHECKED(UFoliageType, RandomYaw));
+	Scaling = DetailLayoutBuilder.GetProperty(GET_MEMBER_NAME_CHECKED(UFoliageType, Scaling));
+	ScaleX = DetailLayoutBuilder.GetProperty(GET_MEMBER_NAME_CHECKED(UFoliageType, ScaleX));
+	ScaleY = DetailLayoutBuilder.GetProperty(GET_MEMBER_NAME_CHECKED(UFoliageType, ScaleY));
+	ScaleZ = DetailLayoutBuilder.GetProperty(GET_MEMBER_NAME_CHECKED(UFoliageType, ScaleZ));
+	ZOffset = DetailLayoutBuilder.GetProperty(GET_MEMBER_NAME_CHECKED(UFoliageType, ZOffset));
+	RandomPitchAngle = DetailLayoutBuilder.GetProperty(GET_MEMBER_NAME_CHECKED(UFoliageType, RandomPitchAngle));
+	GroundSlopeAngle = DetailLayoutBuilder.GetProperty(GET_MEMBER_NAME_CHECKED(UFoliageType, GroundSlopeAngle));
+	Height = DetailLayoutBuilder.GetProperty(GET_MEMBER_NAME_CHECKED(UFoliageType, Height));
+	LandscapeLayers = DetailLayoutBuilder.GetProperty(GET_MEMBER_NAME_CHECKED(UFoliageType, LandscapeLayers));
+	MinimumLayerWeight = DetailLayoutBuilder.GetProperty(GET_MEMBER_NAME_CHECKED(UFoliageType, MinimumLayerWeight));
+	CollisionWithWorld = DetailLayoutBuilder.GetProperty(GET_MEMBER_NAME_CHECKED(UFoliageType, CollisionWithWorld));
+	CollisionScale = DetailLayoutBuilder.GetProperty(GET_MEMBER_NAME_CHECKED(UFoliageType, CollisionScale));
+	VertexColorMask = DetailLayoutBuilder.GetProperty(GET_MEMBER_NAME_CHECKED(UFoliageType, VertexColorMask));
+	VertexColorMaskInvert = DetailLayoutBuilder.GetProperty(GET_MEMBER_NAME_CHECKED(UFoliageType, VertexColorMaskInvert));
+	VertexColorMaskThreshold = DetailLayoutBuilder.GetProperty(GET_MEMBER_NAME_CHECKED(UFoliageType, VertexColorMaskThreshold));
+	
+	//
+	ReapplyDensityAmount = DetailLayoutBuilder.GetProperty(GET_MEMBER_NAME_CHECKED(UFoliageType, ReapplyDensityAmount));
+	ReapplyDensity = DetailLayoutBuilder.GetProperty(GET_MEMBER_NAME_CHECKED(UFoliageType, ReapplyDensity));
+	ReapplyRadius = DetailLayoutBuilder.GetProperty(GET_MEMBER_NAME_CHECKED(UFoliageType, ReapplyRadius));
+	ReapplyAlignToNormal = DetailLayoutBuilder.GetProperty(GET_MEMBER_NAME_CHECKED(UFoliageType, ReapplyAlignToNormal));
+	ReapplyRandomYaw = DetailLayoutBuilder.GetProperty(GET_MEMBER_NAME_CHECKED(UFoliageType, ReapplyRandomYaw));
+	ReapplyScaleX = DetailLayoutBuilder.GetProperty(GET_MEMBER_NAME_CHECKED(UFoliageType, ReapplyScaleX));
+	ReapplyScaleY = DetailLayoutBuilder.GetProperty(GET_MEMBER_NAME_CHECKED(UFoliageType, ReapplyScaleY));
+	ReapplyScaleZ = DetailLayoutBuilder.GetProperty(GET_MEMBER_NAME_CHECKED(UFoliageType, ReapplyScaleZ));
+	ReapplyRandomPitchAngle = DetailLayoutBuilder.GetProperty(GET_MEMBER_NAME_CHECKED(UFoliageType, ReapplyRandomPitchAngle));
+	ReapplyGroundSlopeAngle = DetailLayoutBuilder.GetProperty(GET_MEMBER_NAME_CHECKED(UFoliageType, ReapplyGroundSlope));
+	ReapplyHeight = DetailLayoutBuilder.GetProperty(GET_MEMBER_NAME_CHECKED(UFoliageType, ReapplyHeight));
+	ReapplyLandscapeLayer = DetailLayoutBuilder.GetProperty(GET_MEMBER_NAME_CHECKED(UFoliageType, ReapplyLandscapeLayer));
+	ReapplyZOffset = DetailLayoutBuilder.GetProperty(GET_MEMBER_NAME_CHECKED(UFoliageType, ReapplyZOffset));
+	ReapplyCollisionWithWorld = DetailLayoutBuilder.GetProperty(GET_MEMBER_NAME_CHECKED(UFoliageType, ReapplyCollisionWithWorld));
+	ReapplyVertexColorMask = DetailLayoutBuilder.GetProperty(GET_MEMBER_NAME_CHECKED(UFoliageType, ReapplyVertexColorMask));
+	
+	// Density
+	// TODO: Reapply density
+	AddFoliageProperty(PaintingCategory, Density, InvalidProperty, 
+		TAttribute<EVisibility>(), 
+		TAttribute<bool>::Create(TAttribute<bool>::FGetter::CreateSP(this, &FFoliageTypePaintingCustomization::IsReapplyPropertyEnabled, ReapplyDensity)));
+
+	// Radius
+	AddFoliageProperty(PaintingCategory, Radius, ReapplyRadius, TAttribute<EVisibility>(), TAttribute<bool>());
+
+	// Align to normal
+	AddFoliageProperty(PaintingCategory, AlignToNormal, ReapplyAlignToNormal, TAttribute<EVisibility>(), TAttribute<bool>());
+	AddFoliageProperty(PaintingCategory, AlignMaxAngle, InvalidProperty,
+		TAttribute<EVisibility>::Create(TAttribute<EVisibility>::FGetter::CreateSP(this, &FFoliageTypePaintingCustomization::GetAlignMaxAngleVisibility)),
+		TAttribute<bool>::Create(TAttribute<bool>::FGetter::CreateSP(this, &FFoliageTypePaintingCustomization::IsReapplyPropertyEnabled, ReapplyAlignToNormal))); 
+		
+	// Random Yaw
+	AddFoliageProperty(PaintingCategory, RandomYaw, ReapplyRandomYaw, TAttribute<EVisibility>(), TAttribute<bool>());
+		
+	// Scaling
+	AddFoliageProperty(PaintingCategory, Scaling, InvalidProperty, TAttribute<EVisibility>(), TAttribute<bool>());
+	
+	AddFoliageProperty(PaintingCategory, ScaleX, ReapplyScaleX, 
+		TAttribute<EVisibility>::Create(TAttribute<EVisibility>::FGetter::CreateSP(this, &FFoliageTypePaintingCustomization::GetScaleXVisibility)), TAttribute<bool>());
+	
+	AddFoliageProperty(PaintingCategory, ScaleY, ReapplyScaleY, 
+		TAttribute<EVisibility>::Create(TAttribute<EVisibility>::FGetter::CreateSP(this, &FFoliageTypePaintingCustomization::GetScaleYVisibility)), TAttribute<bool>());
+	
+	AddFoliageProperty(PaintingCategory, ScaleZ, ReapplyScaleZ, 
+		TAttribute<EVisibility>::Create(TAttribute<EVisibility>::FGetter::CreateSP(this, &FFoliageTypePaintingCustomization::GetScaleZVisibility)), TAttribute<bool>());
+
+	// ZOffset
+	AddFoliageProperty(PaintingCategory, ZOffset, ReapplyZOffset, TAttribute<EVisibility>(), TAttribute<bool>());
+	
+	// Random pitch
+	AddFoliageProperty(PaintingCategory, RandomPitchAngle, ReapplyRandomPitchAngle, TAttribute<EVisibility>(), TAttribute<bool>());
+
+	// Ground slope
+	AddFoliageProperty(PaintingCategory, GroundSlopeAngle, ReapplyGroundSlopeAngle, TAttribute<EVisibility>(), TAttribute<bool>());
+
+	// Height
+	AddFoliageProperty(PaintingCategory, Height, ReapplyHeight, TAttribute<EVisibility>(), TAttribute<bool>());
+
+	// Landscape layers
+	AddFoliageProperty(PaintingCategory, LandscapeLayers, ReapplyLandscapeLayer, TAttribute<EVisibility>(), TAttribute<bool>());
+	AddFoliageProperty(PaintingCategory, MinimumLayerWeight, InvalidProperty, 
+		TAttribute<EVisibility>(), 
+		TAttribute<bool>::Create(TAttribute<bool>::FGetter::CreateSP(this, &FFoliageTypePaintingCustomization::IsReapplyPropertyEnabled, ReapplyLandscapeLayer))); 
+
+	// Collision with world
+	AddFoliageProperty(PaintingCategory, CollisionWithWorld, ReapplyCollisionWithWorld, TAttribute<EVisibility>(), TAttribute<bool>());
+	AddFoliageProperty(PaintingCategory, CollisionScale, InvalidProperty, 
+		TAttribute<EVisibility>::Create(TAttribute<EVisibility>::FGetter::CreateSP(this, &FFoliageTypePaintingCustomization::GetCollisionScaleVisibility)),
+		TAttribute<bool>::Create(TAttribute<bool>::FGetter::CreateSP(this, &FFoliageTypePaintingCustomization::IsReapplyPropertyEnabled, ReapplyCollisionWithWorld))); 
+
+	// Vertex color mask
+	AddFoliageProperty(PaintingCategory, VertexColorMask, ReapplyVertexColorMask, TAttribute<EVisibility>(), TAttribute<bool>());
+	
+	AddFoliageProperty(PaintingCategory, VertexColorMaskInvert, InvalidProperty, 
+		TAttribute<EVisibility>::Create(TAttribute<EVisibility>::FGetter::CreateSP(this, &FFoliageTypePaintingCustomization::GetVertexColorMaskDetailsVisibility)),
+		TAttribute<bool>::Create(TAttribute<bool>::FGetter::CreateSP(this, &FFoliageTypePaintingCustomization::IsReapplyPropertyEnabled, ReapplyVertexColorMask))); 
+	
+	AddFoliageProperty(PaintingCategory, VertexColorMaskThreshold, InvalidProperty, 
+		TAttribute<EVisibility>::Create(TAttribute<EVisibility>::FGetter::CreateSP(this, &FFoliageTypePaintingCustomization::GetVertexColorMaskDetailsVisibility)), 
+		TAttribute<bool>::Create(TAttribute<bool>::FGetter::CreateSP(this, &FFoliageTypePaintingCustomization::IsReapplyPropertyEnabled, ReapplyVertexColorMask))); 
+}
+
+IDetailPropertyRow& FFoliageTypePaintingCustomization::AddFoliageProperty(
+		IDetailCategoryBuilder& Category,
+		TSharedPtr<IPropertyHandle>& Property, 
+		TSharedPtr<IPropertyHandle>& ReapplyProperty,
+		const TAttribute<EVisibility>& InVisibility,
+		const TAttribute<bool>& InEnabled)
+{
+	IDetailPropertyRow& PropertyRow = Category.AddProperty(Property);
+
+	if (ReapplyProperty.IsValid())
+	{
+		TSharedPtr<SWidget> NameWidget;
+		TSharedPtr<SWidget> ValueWidget;
+		FDetailWidgetRow	Row;
+		PropertyRow.GetDefaultWidgets(NameWidget, ValueWidget, Row);
+
+		auto IsEnabled = TAttribute<bool>::Create(TAttribute<bool>::FGetter::CreateSP(this, &FFoliageTypePaintingCustomization::IsReapplyPropertyEnabled, ReapplyProperty));
+		ValueWidget->SetEnabled(IsEnabled);
+		NameWidget->SetEnabled(IsEnabled);
+
+		PropertyRow.CustomWidget()
+		.NameContent()
+		.MinDesiredWidth(Row.NameWidget.MinWidth)
+		.MaxDesiredWidth(Row.NameWidget.MaxWidth)
+		[
+			SNew(SHorizontalBox)
+			+SHorizontalBox::Slot()
+			.AutoWidth()
+			[
+				SNew(SCheckBox)
+				.IsChecked(this, &FFoliageTypePaintingCustomization::GetReapplyPropertyState, ReapplyProperty)
+				.OnCheckStateChanged(this, &FFoliageTypePaintingCustomization::OnReapplyPropertyStateChanged, ReapplyProperty)
+				.IsChecked(this, &FFoliageTypePaintingCustomization::GetReapplyPropertyState, ReapplyProperty)
+				.Visibility(this, &FFoliageTypePaintingCustomization::GetReapplyModeVisibility)
+				.ToolTipText(ReapplyProperty->GetToolTipText())
+			]
+			+SHorizontalBox::Slot()
+			.AutoWidth()
+			[
+				NameWidget.ToSharedRef()
+			]
+		]
+		.ValueContent()
+		.MinDesiredWidth(Row.ValueWidget.MinWidth)
+		.MaxDesiredWidth(Row.ValueWidget.MaxWidth)
+		[
+			ValueWidget.ToSharedRef()
+		];
+	}
+	else 
+	{
+		if (InEnabled.IsSet())
+		{
+			PropertyRow.IsEnabled(InEnabled);
+		}
+	}
+
+	if (InVisibility.IsSet())
+	{
+		PropertyRow.Visibility(InVisibility);
+	}
+	
+	//
+	return PropertyRow;
+}
+
+void FFoliageTypePaintingCustomization::HideFoliageCategory(IDetailLayoutBuilder& DetailLayoutBuilder, FName CategoryName)
+{
+	TArray<TSharedRef<IPropertyHandle>> CategoryProperties;
+	DetailLayoutBuilder.EditCategory(CategoryName).GetDefaultProperties(CategoryProperties, true, true);
+
+	for (auto& PropertyHandle : CategoryProperties)
+	{
+		DetailLayoutBuilder.HideProperty(PropertyHandle);
+	}
+}
+
+EVisibility FFoliageTypePaintingCustomization::GetAlignMaxAngleVisibility() const
+{
+	bool bAlightToNormal;
+	if (AlignToNormal->GetValue(bAlightToNormal) == FPropertyAccess::Success)
+	{
+		return bAlightToNormal ? EVisibility::Visible : EVisibility::Collapsed;
+	}
+
+	return EVisibility::Visible;
+}
+
+EVisibility FFoliageTypePaintingCustomization::GetScaleXVisibility() const
+{
+	return EVisibility::Visible;
+}
+
+EVisibility FFoliageTypePaintingCustomization::GetScaleYVisibility() const
+{
+	uint8 ScalingValue;
+	if (Scaling->GetValue(ScalingValue) == FPropertyAccess::Success)
+	{
+		return (ScalingValue == (uint8)EFoliageScaling::Uniform) ? EVisibility::Collapsed : EVisibility::Visible;
+	}
+	return EVisibility::Visible;
+}
+
+EVisibility FFoliageTypePaintingCustomization::GetScaleZVisibility() const
+{
+	uint8 ScalingValue;
+	if (Scaling->GetValue(ScalingValue) == FPropertyAccess::Success)
+	{
+		return (ScalingValue == (uint8)EFoliageScaling::Uniform) ? EVisibility::Collapsed : EVisibility::Visible;
+	}
+	return EVisibility::Visible;
+}
+
+EVisibility FFoliageTypePaintingCustomization::GetCollisionScaleVisibility() const
+{
+	bool bCollisionWithWorld = false;
+	if (CollisionWithWorld->GetValue(bCollisionWithWorld) == FPropertyAccess::Success)
+	{
+		return bCollisionWithWorld ? EVisibility::Visible : EVisibility::Collapsed;
+	}
+	return EVisibility::Visible;
+}
+
+EVisibility FFoliageTypePaintingCustomization::GetVertexColorMaskDetailsVisibility() const
+{
+	uint8 ColorMask;
+	if (VertexColorMask->GetValue(ColorMask) == FPropertyAccess::Success)
+	{
+		return ColorMask != FOLIAGEVERTEXCOLORMASK_Disabled ? EVisibility::Visible : EVisibility::Collapsed;
+	}
+	return EVisibility::Visible;
+}
+
+EVisibility FFoliageTypePaintingCustomization::GetReapplyModeVisibility() const
+{
+	return FoliageEditMode->UISettings.GetReapplyToolSelected() ? EVisibility::Visible : EVisibility::Collapsed;
+}
+
+ECheckBoxState FFoliageTypePaintingCustomization::GetReapplyPropertyState(TSharedPtr<IPropertyHandle> ReapplyProperty) const
+{
+	bool bState;
+	if (ReapplyProperty->GetValue(bState) == FPropertyAccess::Success)
+	{
+		return bState ? ECheckBoxState::Checked : ECheckBoxState::Unchecked;
+	}
+		
+	return ECheckBoxState::Undetermined;
+}
+
+void FFoliageTypePaintingCustomization::OnReapplyPropertyStateChanged(ECheckBoxState CheckState, TSharedPtr<IPropertyHandle> ReapplyProperty)
+{
+	if (CheckState != ECheckBoxState::Undetermined)
+	{
+		ReapplyProperty->SetValue(CheckState == ECheckBoxState::Checked);
+	}
+}
+
+bool FFoliageTypePaintingCustomization::IsReapplyPropertyEnabled(TSharedPtr<IPropertyHandle> ReapplyProperty) const
+{
+	if (FoliageEditMode->UISettings.GetReapplyToolSelected())
+	{
+		bool bState;
+		if (ReapplyProperty->GetValue(bState) == FPropertyAccess::Success)
+		{
+			return bState;
+		}
+	}
+
+	return true;
+}
+
+#undef LOCTEXT_NAMESPACE
\ No newline at end of file
diff --git a/Engine/Source/Editor/FoliageEdit/Private/FoliageTypePaintingCustomization.h b/Engine/Source/Editor/FoliageEdit/Private/FoliageTypePaintingCustomization.h
new file mode 100644
--- /dev/null
+++ b/Engine/Source/Editor/FoliageEdit/Private/FoliageTypePaintingCustomization.h
@@ -0,0 +1,84 @@
+// Copyright 1998-2015 Epic Games, Inc. All Rights Reserved.
+#pragma once
+
+class FEdModeFoliage;
+class IDetailCategoryBuilder;
+class IDetailPropertyRow;
+
+/////////////////////////////////////////////////////
+// FFoliageTypePaintingCustomization
+class FFoliageTypePaintingCustomization : public IDetailCustomization
+{
+public:
+	// Makes a new instance of this detail layout class
+	static TSharedRef<IDetailCustomization> MakeInstance(FEdModeFoliage* InFoliageEditMode);
+
+	// IDetailCustomization interface
+	virtual void CustomizeDetails(IDetailLayoutBuilder& DetailLayout) override;
+	// End of IDetailCustomization interface
+private:
+	FFoliageTypePaintingCustomization(FEdModeFoliage* InFoliageEditMode);
+	
+	EVisibility GetAlignMaxAngleVisibility() const;
+	EVisibility GetScaleXVisibility() const;
+	EVisibility GetScaleYVisibility() const;
+	EVisibility GetScaleZVisibility() const;
+	EVisibility GetCollisionScaleVisibility() const;
+	EVisibility GetVertexColorMaskDetailsVisibility() const;
+	
+	EVisibility GetReapplyModeVisibility() const;
+	ECheckBoxState GetReapplyPropertyState(TSharedPtr<IPropertyHandle> ReapplyProperty) const;
+	void OnReapplyPropertyStateChanged(ECheckBoxState CheckState, TSharedPtr<IPropertyHandle> ReapplyProperty);
+	bool IsReapplyPropertyEnabled(TSharedPtr<IPropertyHandle> ReapplyProperty) const;
+
+	IDetailPropertyRow& AddFoliageProperty(
+		IDetailCategoryBuilder& Category,
+		TSharedPtr<IPropertyHandle>& Property, 
+		TSharedPtr<IPropertyHandle>& ReapplyProperty,
+		const TAttribute<EVisibility>& InVisibility,
+		const TAttribute<bool>& InEnabled);
+
+	static void HideFoliageCategory(IDetailLayoutBuilder& DetailLayout,	FName CategoryName);
+	
+private:
+	/** Pointer to the foliage edit mode. */
+	FEdModeFoliage* FoliageEditMode;
+
+	TSharedPtr<IPropertyHandle> Density;
+	TSharedPtr<IPropertyHandle> Radius;
+	TSharedPtr<IPropertyHandle> AlignToNormal;
+	TSharedPtr<IPropertyHandle> AlignMaxAngle;
+	TSharedPtr<IPropertyHandle> RandomYaw;
+	TSharedPtr<IPropertyHandle> Scaling;
+	TSharedPtr<IPropertyHandle> ScaleX;
+	TSharedPtr<IPropertyHandle> ScaleY;
+	TSharedPtr<IPropertyHandle> ScaleZ;
+	TSharedPtr<IPropertyHandle> ZOffset;
+	TSharedPtr<IPropertyHandle> RandomPitchAngle;
+	TSharedPtr<IPropertyHandle> GroundSlopeAngle;
+	TSharedPtr<IPropertyHandle> Height;
+	TSharedPtr<IPropertyHandle>	LandscapeLayers;
+	TSharedPtr<IPropertyHandle>	MinimumLayerWeight;
+	TSharedPtr<IPropertyHandle> CollisionWithWorld;
+	TSharedPtr<IPropertyHandle> CollisionScale;
+	TSharedPtr<IPropertyHandle> VertexColorMask;
+	TSharedPtr<IPropertyHandle> VertexColorMaskInvert;
+	TSharedPtr<IPropertyHandle> VertexColorMaskThreshold;
+
+	//
+	TSharedPtr<IPropertyHandle> ReapplyDensity;
+	TSharedPtr<IPropertyHandle> ReapplyDensityAmount;
+	TSharedPtr<IPropertyHandle> ReapplyRadius;
+	TSharedPtr<IPropertyHandle> ReapplyAlignToNormal;
+	TSharedPtr<IPropertyHandle> ReapplyRandomYaw;
+	TSharedPtr<IPropertyHandle> ReapplyScaleX;
+	TSharedPtr<IPropertyHandle> ReapplyScaleY;
+	TSharedPtr<IPropertyHandle> ReapplyScaleZ;
+	TSharedPtr<IPropertyHandle> ReapplyRandomPitchAngle;
+	TSharedPtr<IPropertyHandle> ReapplyGroundSlopeAngle;
+	TSharedPtr<IPropertyHandle> ReapplyHeight;
+	TSharedPtr<IPropertyHandle> ReapplyLandscapeLayer;
+	TSharedPtr<IPropertyHandle> ReapplyZOffset;
+	TSharedPtr<IPropertyHandle> ReapplyCollisionWithWorld;
+	TSharedPtr<IPropertyHandle> ReapplyVertexColorMask;
+};
diff --git a/Engine/Source/Editor/FoliageEdit/Private/SFoliageEdit.cpp b/Engine/Source/Editor/FoliageEdit/Private/SFoliageEdit.cpp
--- a/Engine/Source/Editor/FoliageEdit/Private/SFoliageEdit.cpp
+++ b/Engine/Source/Editor/FoliageEdit/Private/SFoliageEdit.cpp
@@ -4,16 +4,21 @@
 
 #include "SScaleBox.h"
 
+#include "FoliageType.h"
+#include "FoliageEdMode.h"
 #include "SFoliageEdit.h"
 #include "FoliageEditActions.h"
-#include "FoliageEdMode.h"
-#include "FoliageType.h"
 #include "Editor/UnrealEd/Public/AssetThumbnail.h"
 #include "Editor/UnrealEd/Public/DragAndDrop/AssetDragDropOp.h"
+#include "Editor/ContentBrowser/Public/ContentBrowserModule.h"
 #include "Editor/UnrealEd/Public/AssetSelection.h"
 #include "Editor/IntroTutorials/Public/IIntroTutorials.h"
-#include "SFoliageEditMeshDisplayItem.h"
+#include "Editor/PropertyEditor/Public/IDetailsView.h"
+#include "Editor/PropertyEditor/Public/PropertyEditorModule.h"
 #include "Engine/StaticMesh.h"
+#include "SFoliageEditMeshRow.h"
+#include "FoliageTypePaintingCustomization.h"
+
 
 #define LOCTEXT_NAMESPACE "FoliageEd_Mode"
 
@@ -78,6 +83,7 @@ private:
 	bool bIsDragOn;
 };
 
+
 BEGIN_SLATE_FUNCTION_BUILD_OPTIMIZATION
 void SFoliageEdit::Construct(const FArguments& InArgs)
 {
@@ -89,7 +95,13 @@ void SFoliageEdit::Construct(const FArguments& InArgs)
 
 	BindCommands();
 
-	AssetThumbnailPool = MakeShareable(new FAssetThumbnailPool(512, TAttribute<bool>::Create(TAttribute<bool>::FGetter::CreateSP(this, &SFoliageEdit::IsHovered))));
+	FPropertyEditorModule& PropertyModule = FModuleManager::LoadModuleChecked<FPropertyEditorModule>("PropertyEditor");
+	FDetailsViewArgs Args(false, false, false, FDetailsViewArgs::HideNameArea, true);
+	Args.bShowActorLabel = false;
+	MeshDetailsWidget = PropertyModule.CreateDetailView(Args);
+	MeshDetailsWidget->RegisterInstancedCustomPropertyLayout(UFoliageType::StaticClass(), 
+		FOnGetDetailCustomizationInstance::CreateStatic(&FFoliageTypePaintingCustomization::MakeInstance, FoliageEditMode)
+		);
 
 	// Everything (or almost) uses this padding, change it to expand the padding.
 	FMargin StandardPadding(0.0f, 4.0f, 0.0f, 4.0f);
@@ -99,214 +111,240 @@ void SFoliageEdit::Construct(const FArguments& InArgs)
 		SNew(SOverlay)
 
 		+ SOverlay::Slot()
+		[
+			SNew(SVerticalBox)
+
+			+ SVerticalBox::Slot()
+			.AutoHeight()
+			.Padding(StandardPadding)
 			[
-			SNew(SScrollBox)
-			+ SScrollBox::Slot()
-			.Padding(2.0f, 0)
+				BuildToolBar()
+			]
+
+			+ SVerticalBox::Slot()
+			.AutoHeight()
 			[
-				SNew(SVerticalBox)
+				SNew(SHorizontalBox)
+				.ToolTipText(LOCTEXT("BrushSize_Tooltip", "The size of the foliage brush"))
 
-				+ SVerticalBox::Slot()
-				.AutoHeight()
+				+ SHorizontalBox::Slot()
+				.FillWidth(1.0f)
 				.Padding(StandardPadding)
 				[
-					BuildToolBar()
+					SNew(STextBlock)
+					.Visibility(this, &SFoliageEdit::GetVisibility_Radius)
+					.Text(LOCTEXT("BrushSize", "Brush Size"))
 				]
-
-				+ SVerticalBox::Slot()
-				.AutoHeight()
+				+ SHorizontalBox::Slot()
+				.FillWidth(2.0f)
+				.Padding(StandardPadding)
 				[
-					SNew(SHorizontalBox)
-					.ToolTipText(LOCTEXT("BrushSize_Tooltip", "The size of the foliage brush"))
-
-					+ SHorizontalBox::Slot()
-					.FillWidth(1.0f)
-					.Padding(StandardPadding)
-					[
-						SNew(STextBlock)
-						.Visibility(this, &SFoliageEdit::GetVisibility_Radius)
-						.Text(LOCTEXT("BrushSize", "Brush Size"))
-					]
-					+ SHorizontalBox::Slot()
-					.FillWidth(2.0f)
-					.Padding(StandardPadding)
-					[
-						SNew(SSpinBox<float>)
-						.Visibility(this, &SFoliageEdit::GetVisibility_Radius)
-						.MinValue(0.0f)
-						.MaxValue(65536.0f)
-						.MaxSliderValue(8192.0f)
-						.SliderExponent(3.0f)
-						.Value(this, &SFoliageEdit::GetRadius)
-						.OnValueChanged(this, &SFoliageEdit::SetRadius)
-					]
+					SNew(SSpinBox<float>)
+					.Visibility(this, &SFoliageEdit::GetVisibility_Radius)
+					.MinValue(0.0f)
+					.MaxValue(65536.0f)
+					.MaxSliderValue(8192.0f)
+					.SliderExponent(3.0f)
+					.Value(this, &SFoliageEdit::GetRadius)
+					.OnValueChanged(this, &SFoliageEdit::SetRadius)
 				]
+			]
 
-				+ SVerticalBox::Slot()
-				.AutoHeight()
-				[
-					SNew(SHorizontalBox)
-					.ToolTipText(LOCTEXT("PaintDensity_Tooltip", "The density of foliage to paint. This is a multiplier for the individual foliage type's density specifier."))
+			+ SVerticalBox::Slot()
+			.AutoHeight()
+			[
+				SNew(SHorizontalBox)
+				.ToolTipText(LOCTEXT("PaintDensity_Tooltip", "The density of foliage to paint. This is a multiplier for the individual foliage type's density specifier."))
 
-					+ SHorizontalBox::Slot()
-					.FillWidth(1.0f)
-					.Padding(StandardPadding)
-					[
-						SNew(STextBlock)
-						.Visibility(this, &SFoliageEdit::GetVisibility_PaintDensity)
-						.Text(LOCTEXT("PaintDensity", "Paint Density"))
-					]
-					+ SHorizontalBox::Slot()
-					.FillWidth(2.0f)
-					.Padding(StandardPadding)
-					[
-						SNew(SSpinBox<float>)
-						.Visibility(this, &SFoliageEdit::GetVisibility_PaintDensity)
-						.MinValue(0.0f)
-						.MaxValue(2.0f)
-						.MaxSliderValue(1.0f)
-						.Value(this, &SFoliageEdit::GetPaintDensity)
-						.OnValueChanged(this, &SFoliageEdit::SetPaintDensity)
-					]
+				+ SHorizontalBox::Slot()
+				.FillWidth(1.0f)
+				.Padding(StandardPadding)
+				[
+					SNew(STextBlock)
+					.Visibility(this, &SFoliageEdit::GetVisibility_PaintDensity)
+					.Text(LOCTEXT("PaintDensity", "Paint Density"))
 				]
-				+ SVerticalBox::Slot()
-				.AutoHeight()
+				+ SHorizontalBox::Slot()
+				.FillWidth(2.0f)
+				.Padding(StandardPadding)
 				[
-					SNew(SHorizontalBox)
-					.ToolTipText(LOCTEXT("EraseDensity_Tooltip", "The density of foliage to leave behind when erasing with the Shift key held. 0 will remove all foliage."))
-
-					+ SHorizontalBox::Slot()
-					.FillWidth(1.0f)
-					.Padding(StandardPadding)
-					[
-						SNew(STextBlock)
-						.Visibility(this, &SFoliageEdit::GetVisibility_EraseDensity)
-						.Text(LOCTEXT("EraseDensity", "Erase Density"))
-					]
-					+ SHorizontalBox::Slot()
-					.FillWidth(2.0f)
-					.Padding(StandardPadding)
-					[
-						SNew(SSpinBox<float>)
-						.Visibility(this, &SFoliageEdit::GetVisibility_EraseDensity)
-						.MinValue(0.0f)
-						.MaxValue(2.0f)
-						.MaxSliderValue(1.0f)
-						.Value(this, &SFoliageEdit::GetEraseDensity)
-						.OnValueChanged(this, &SFoliageEdit::SetEraseDensity)
-					]
+					SNew(SSpinBox<float>)
+					.Visibility(this, &SFoliageEdit::GetVisibility_PaintDensity)
+					.MinValue(0.0f)
+					.MaxValue(2.0f)
+					.MaxSliderValue(1.0f)
+					.Value(this, &SFoliageEdit::GetPaintDensity)
+					.OnValueChanged(this, &SFoliageEdit::SetPaintDensity)
 				]
+			]
 
-				+ SVerticalBox::Slot()
-				.AutoHeight()
+			+ SVerticalBox::Slot()
+			.AutoHeight()
+			[
+				SNew(SHorizontalBox)
+				.ToolTipText(LOCTEXT("EraseDensity_Tooltip", "The density of foliage to leave behind when erasing with the Shift key held. 0 will remove all foliage."))
+
+				+ SHorizontalBox::Slot()
+				.FillWidth(1.0f)
 				.Padding(StandardPadding)
 				[
 					SNew(STextBlock)
-					.Visibility(this, &SFoliageEdit::GetVisibility_Filters)
-					.Text(LOCTEXT("Filter", "Filter"))
+					.Visibility(this, &SFoliageEdit::GetVisibility_EraseDensity)
+					.Text(LOCTEXT("EraseDensity", "Erase Density"))
 				]
-
-				+ SVerticalBox::Slot()
-				.AutoHeight()
+				+ SHorizontalBox::Slot()
+				.FillWidth(2.0f)
 				.Padding(StandardPadding)
 				[
-					SNew(SWrapBox)
-					.UseAllottedWidth(true)
+					SNew(SSpinBox<float>)
+					.Visibility(this, &SFoliageEdit::GetVisibility_EraseDensity)
+					.MinValue(0.0f)
+					.MaxValue(2.0f)
+					.MaxSliderValue(1.0f)
+					.Value(this, &SFoliageEdit::GetEraseDensity)
+					.OnValueChanged(this, &SFoliageEdit::SetEraseDensity)
+				]
+			]
+
+			+ SVerticalBox::Slot()
+			.AutoHeight()
+			.Padding(StandardPadding)
+			[
+				SNew(STextBlock)
+				.Visibility(this, &SFoliageEdit::GetVisibility_Filters)
+				.Text(LOCTEXT("Filter", "Filter"))
+			]
+
+			+ SVerticalBox::Slot()
+			.AutoHeight()
+			.Padding(StandardPadding)
+			[
+				SNew(SWrapBox)
+				.UseAllottedWidth(true)
 
-					+ SWrapBox::Slot()
-					.Padding(0, 0, 16, 0)
+				+ SWrapBox::Slot()
+				.Padding(0, 0, 16, 0)
+				[
+					SNew(SCheckBox)
+					.Visibility(this, &SFoliageEdit::GetVisibility_Filters)
+					.OnCheckStateChanged(this, &SFoliageEdit::OnCheckStateChanged_Landscape)
+					.IsChecked(this, &SFoliageEdit::GetCheckState_Landscape)
+					.ToolTipText(LOCTEXT("FilterLandscape_Tooltip", "Place foliage on Landscape"))
 					[
-						SNew(SCheckBox)
-						.Visibility(this, &SFoliageEdit::GetVisibility_Filters)
-						.OnCheckStateChanged(this, &SFoliageEdit::OnCheckStateChanged_Landscape)
-						.IsChecked(this, &SFoliageEdit::GetCheckState_Landscape)
-						.ToolTipText(LOCTEXT("FilterLandscape_Tooltip", "Place foliage on Landscape"))
-						[
-							SNew(STextBlock)
-							.Text(LOCTEXT("Landscape", "Landscape"))
-						]
+						SNew(STextBlock)
+						.Text(LOCTEXT("Landscape", "Landscape"))
 					]
+				]
 
-					+ SWrapBox::Slot()
-					.Padding(0, 0, 16, 0)
+				+ SWrapBox::Slot()
+				.Padding(0, 0, 16, 0)
+				[
+					SNew(SCheckBox)
+					.Visibility(this, &SFoliageEdit::GetVisibility_Filters)
+					.OnCheckStateChanged(this, &SFoliageEdit::OnCheckStateChanged_StaticMesh)
+					.IsChecked(this, &SFoliageEdit::GetCheckState_StaticMesh)
+					.ToolTipText(LOCTEXT("FilterStaticMesh_Tooltip", "Place foliage on StaticMeshes"))
 					[
-						SNew(SCheckBox)
-						.Visibility(this, &SFoliageEdit::GetVisibility_Filters)
-						.OnCheckStateChanged(this, &SFoliageEdit::OnCheckStateChanged_StaticMesh)
-						.IsChecked(this, &SFoliageEdit::GetCheckState_StaticMesh)
-						.ToolTipText(LOCTEXT("FilterStaticMesh_Tooltip", "Place foliage on StaticMeshes"))
-						[
-							SNew(STextBlock)
-							.Text(LOCTEXT("StaticMeshes", "Static Meshes"))
-						]
+						SNew(STextBlock)
+						.Text(LOCTEXT("StaticMeshes", "Static Meshes"))
 					]
+				]
 
-					+ SWrapBox::Slot()
-					.Padding(0, 0, 16, 0)
+				+ SWrapBox::Slot()
+				.Padding(0, 0, 16, 0)
+				[
+					SNew(SCheckBox)
+					.Visibility(this, &SFoliageEdit::GetVisibility_Filters)
+					.OnCheckStateChanged(this, &SFoliageEdit::OnCheckStateChanged_BSP)
+					.IsChecked(this, &SFoliageEdit::GetCheckState_BSP)
+					.ToolTipText(LOCTEXT("FilterBSP_Tooltip", "Place foliage on BSP"))
 					[
-						SNew(SCheckBox)
-						.Visibility(this, &SFoliageEdit::GetVisibility_Filters)
-						.OnCheckStateChanged(this, &SFoliageEdit::OnCheckStateChanged_BSP)
-						.IsChecked(this, &SFoliageEdit::GetCheckState_BSP)
-						.ToolTipText(LOCTEXT("FilterBSP_Tooltip", "Place foliage on BSP"))
-						[
-							SNew(STextBlock)
-							.Text(LOCTEXT("BSP", "BSP"))
-						]
+						SNew(STextBlock)
+						.Text(LOCTEXT("BSP", "BSP"))
 					]
+				]
 
-					+ SWrapBox::Slot()
-					.Padding(0, 0, 16, 0)
+				+ SWrapBox::Slot()
+				.Padding(0, 0, 16, 0)
+				[
+					SNew(SCheckBox)
+					.Visibility(this, &SFoliageEdit::GetVisibility_Filters)
+					.OnCheckStateChanged(this, &SFoliageEdit::OnCheckStateChanged_Translucent)
+					.IsChecked(this, &SFoliageEdit::GetCheckState_Translucent)
+					.ToolTipText(LOCTEXT("FilterTranslucent_Tooltip", "Place foliage on translucent surfaces"))
 					[
-						SNew(SCheckBox)
-						.Visibility(this, &SFoliageEdit::GetVisibility_Filters)
-						.OnCheckStateChanged(this, &SFoliageEdit::OnCheckStateChanged_Translucent)
-						.IsChecked(this, &SFoliageEdit::GetCheckState_Translucent)
-						.ToolTipText(LOCTEXT("FilterTranslucent_Tooltip", "Place foliage on translucent surfaces"))
-						[
-							SNew(STextBlock)
-							.Text(LOCTEXT("Translucent", "Translucent"))
-						]
+						SNew(STextBlock)
+						.Text(LOCTEXT("Translucent", "Translucent"))
 					]
 				]
+			]
 
-				+ SVerticalBox::Slot()
-				.AutoHeight()
-				[
-					SNew(SHeaderRow)
+			+ SVerticalBox::Slot()
+			.AutoHeight()
+			[
+				SNew(SHeaderRow)
 
-					+ SHeaderRow::Column(TEXT("Meshes"))
-					.DefaultLabel(LOCTEXT("Meshes", "Meshes"))
-				]
+				+ SHeaderRow::Column(TEXT("Meshes"))
+				.DefaultLabel(LOCTEXT("Meshes", "Meshes"))
+			]
 
-				+ SVerticalBox::Slot()
-				.AutoHeight()
-				.Padding(StandardPadding)
+			+ SVerticalBox::Slot()
+			.AutoHeight()
+			.Padding(StandardPadding)
+			[
+				SNew(SBox)
+				.Visibility(this, &SFoliageEdit::GetVisibility_EmptyList)
+				.Padding(FMargin(15, 0))
+				.MinDesiredHeight(30)
 				[
-					SNew(SBox)
-					.Visibility(this, &SFoliageEdit::GetVisibility_EmptyList)
-					.Padding(FMargin(15, 0))
-					.MinDesiredHeight(30)
+					SNew(SScaleBox)
+					.Stretch(EStretch::ScaleToFit)
 					[
-						SNew(SScaleBox)
-						.Stretch(EStretch::ScaleToFit)
-						[
-							SNew(STextBlock)
-							.Text(LOCTEXT("Foliage_DropStatic", "+ Drop Static Meshes Here"))
-							.ToolTipText(LOCTEXT("Foliage_DropStatic_ToolTip", "Drag static meshes from the content browser into this area."))
-						]
+						SNew(STextBlock)
+						.Text(LOCTEXT("Foliage_DropStatic", "+ Drop Foliage Here"))
+						.ToolTipText(LOCTEXT("Foliage_DropStatic_ToolTip", "Drag static meshes/foliage settings from the content browser into this area."))
 					]
 				]
+			]
+
+			+ SVerticalBox::Slot()
+			.FillHeight(1.0f)
+			[
+				SNew(SSplitter)
+				.Orientation(Orient_Vertical)
 
-				+ SVerticalBox::Slot()
-				.AutoHeight()
+				+SSplitter::Slot()
+				.Value(0.3f)
 				[
-					SAssignNew(ItemScrollBox, SScrollBox)
+					SAssignNew(MeshTreeWidget, SFoliageMeshTree)
+					.TreeItemsSource(&FoliageEditMode->GetFoliageMeshList())
+					.SelectionMode(ESelectionMode::Multi)
+					.OnGenerateRow(this, &SFoliageEdit::MeshTreeGenerateRow)
+					.OnGetChildren( this, &SFoliageEdit::MeshTreeGetChildren)
+					.OnContextMenuOpening(this, &SFoliageEdit::ConstructFoliageMeshContextMenu)
+					.OnSelectionChanged(this, &SFoliageEdit::MeshTreeOnSelectionChanged)
+					.HeaderRow
+						(
+							SNew( SHeaderRow ).Visibility(EVisibility::Collapsed)
+							// toggle mesh for painting checkbox
+							+SHeaderRow::Column(FoliageMeshColumns::ColumnID_ToggleMesh).FixedWidth(24.0f)
+							// mesh text label
+							+SHeaderRow::Column(FoliageMeshColumns::ColumnID_MeshLabel).FillWidth(1.0f)
+							// instance count
+							+SHeaderRow::Column(FoliageMeshColumns::ColumnID_InstanceCount).FillWidth(0.3f)
+							// save Foliage asset
+							+SHeaderRow::Column(FoliageMeshColumns::ColumnID_Save).FixedWidth(24.0f)
+						)
+				]
+					
+				+SSplitter::Slot()
+				[
+					MeshDetailsWidget.ToSharedRef()
 				]
 			]
 		]
 
-		// Static Mesh Drop Zone
+		// Foliage Mesh Drop Zone
 		+ SOverlay::Slot()
 		.HAlign(HAlign_Fill)
 		.VAlign(VAlign_Fill)
@@ -332,75 +370,11 @@ END_SLATE_FUNCTION_BUILD_OPTIMIZATION
 
 SFoliageEdit::~SFoliageEdit()
 {
-	// Release all rendering resources being held onto
-	AssetThumbnailPool->ReleaseResources();
-
-	FEditorDelegates::NewCurrentLevel.RemoveAll(this);
-}
-
-TSharedPtr<FAssetThumbnail> SFoliageEdit::CreateThumbnail(UStaticMesh* InStaticMesh)
-{
-	return MakeShareable(new FAssetThumbnail(InStaticMesh, 80, 80, AssetThumbnailPool));
 }
 
 void SFoliageEdit::RefreshFullList()
 {
-	for (int DisplayItemIdx = 0; DisplayItemIdx < DisplayItemList.Num();)
-	{
-		RemoveItemFromScrollbox(DisplayItemList[DisplayItemIdx]);
-	}
-	
-	//FoliageEditMode->PopulateFoliageMeshList();
-	auto& FoliageMeshList = FoliageEditMode->GetFoliageMeshList();
-
-	for (int MeshIdx = 0; MeshIdx < FoliageMeshList.Num(); MeshIdx++)
-	{
-		AddItemToScrollbox(FoliageMeshList[MeshIdx]);
-	}
-}
-
-void SFoliageEdit::AddItemToScrollbox(TSharedPtr<FFoliageMeshUIInfo>& InFoliageInfoToAdd)
-{
-	UFoliageType* Settings = InFoliageInfoToAdd->Settings;
-
-	TSharedRef<SFoliageEditMeshDisplayItem> DisplayItem =
-		SNew(SFoliageEditMeshDisplayItem)
-		.FoliageEditPtr(SharedThis(this))
-		.FoliageSettingsPtr(Settings)
-		.AssetThumbnail(CreateThumbnail(Settings->GetStaticMesh()))
-		.FoliageMeshUIInfo(InFoliageInfoToAdd);
-
-	DisplayItemList.Add(DisplayItem);
-
-	ItemScrollBox->AddSlot()
-	[
-		DisplayItem
-	];
-}
-
-void SFoliageEdit::RemoveItemFromScrollbox(const TSharedPtr<SFoliageEditMeshDisplayItem> InWidgetToRemove)
-{
-	DisplayItemList.Remove(InWidgetToRemove.ToSharedRef());
-
-	ItemScrollBox->RemoveSlot(InWidgetToRemove.ToSharedRef());
-}
-
-void SFoliageEdit::ReplaceItem(const TSharedPtr<SFoliageEditMeshDisplayItem> InDisplayItemToReplaceIn, UStaticMesh* InNewStaticMesh)
-{
-	TArray<FFoliageMeshUIInfoPtr>& FoliageMeshList = FoliageEditMode->GetFoliageMeshList();
-
-	for (FFoliageMeshUIInfoPtr& UIInfo : FoliageMeshList)
-	{
-		UFoliageType* Settings = UIInfo->Settings;
-		if (Settings->GetStaticMesh() == InNewStaticMesh)
-		{
-			// This is the info for the new static mesh. Update the display item.
-			InDisplayItemToReplaceIn->Replace(Settings,
-				CreateThumbnail(Settings->GetStaticMesh()),
-				UIInfo);
-			break;
-		}
-	}
+	MeshTreeWidget->RequestTreeRefresh();
 }
 
 void SFoliageEdit::ClearAllToolSelection()
@@ -445,6 +419,20 @@ void SFoliageEdit::BindCommands()
 		FExecuteAction::CreateSP(this, &SFoliageEdit::OnSetPaintFill),
 		FCanExecuteAction(),
 		FIsActionChecked::CreateSP(this, &SFoliageEdit::IsPaintFillTool));
+
+		//
+	UICommandList->MapAction(
+		Commands.RemoveFoliageType,
+		FExecuteAction::CreateSP(this, &SFoliageEdit::OnRemoveFoliageType),
+		FCanExecuteAction(),
+		FIsActionChecked());
+
+	UICommandList->MapAction(
+		Commands.ShowFoliageTypeInCB,
+		FExecuteAction::CreateSP(this, &SFoliageEdit::OnShowFoliageTypeInCB),
+		FCanExecuteAction(),
+		FIsActionChecked());
+
 }
 
 TSharedRef<SWidget> SFoliageEdit::BuildToolBar()
@@ -593,7 +581,7 @@ float SFoliageEdit::GetEraseDensity() const
 
 EVisibility SFoliageEdit::GetVisibility_EmptyList() const
 {
-	if (DisplayItemList.Num() > 0)
+	if (FoliageEditMode->GetFoliageMeshList().Num() > 0)
 	{
 		return EVisibility::Collapsed;
 	}
@@ -608,7 +596,7 @@ EVisibility SFoliageEdit::GetVisibility_FoliageDropTarget() const
 		TArray<FAssetData> DraggedAssets = AssetUtil::ExtractAssetDataFromDrag(FSlateApplication::Get().GetDragDroppingContent());
 		for ( const FAssetData& AssetData : DraggedAssets )
 		{
-			if ( AssetData.IsValid() && (AssetData.GetClass()->IsChildOf(UStaticMesh::StaticClass()) || AssetData.GetClass()->IsChildOf(UFoliageType::StaticClass())) )
+			if ( AssetData.IsValid() && (AssetData.GetClass()->IsChildOf(UStaticMesh::StaticClass()) ||	AssetData.GetClass()->IsChildOf(UFoliageType::StaticClass())) )
 			{
 				return EVisibility::Visible;
 			}
@@ -620,7 +608,7 @@ EVisibility SFoliageEdit::GetVisibility_FoliageDropTarget() const
 
 EVisibility SFoliageEdit::GetVisibility_NonEmptyList() const
 {
-	if (DisplayItemList.Num() == 0)
+	if (FoliageEditMode->GetFoliageMeshList().Num() == 0)
 	{
 		return EVisibility::Collapsed;
 	}
@@ -628,22 +616,6 @@ EVisibility SFoliageEdit::GetVisibility_NonEmptyList() const
 	return EVisibility::HitTestInvisible;
 }
 
-bool SFoliageEdit::CanAddStaticMesh(const UStaticMesh* const InStaticMesh) const
-{
-	TArray<FFoliageMeshUIInfoPtr>& FoliageMeshList = FoliageEditMode->GetFoliageMeshList();
-	
-	for (FFoliageMeshUIInfoPtr& UIInfo : FoliageMeshList)
-	{
-		UFoliageType* Settings = UIInfo->Settings;
-
-		if (Settings->GetStaticMesh() == InStaticMesh)
-		{
-			return false;
-		}
-	}
-	return true;
-}
-
 FReply SFoliageEdit::OnDrop_ListView(const FGeometry& MyGeometry, const FDragDropEvent& DragDropEvent)
 {
 	TArray<FAssetData> DroppedAssetData = AssetUtil::ExtractAssetDataFromDrag(DragDropEvent);
@@ -654,14 +626,31 @@ FReply SFoliageEdit::OnDrop_ListView(const FGeometry& MyGeometry, const FDragDro
 		UObject* Asset = DroppedAssetData[AssetIdx].GetAsset();
 		GWarn->EndSlowTask();
 
-		if (Asset)
+		UFoliageType* FoliageType = FoliageEditMode->AddFoliageAsset(Asset);
+		if (FoliageType)
 		{
-			FoliageEditMode->AddFoliageAsset(Asset);
+			MeshTreeWidget->RequestTreeRefresh();
 		}
 	}
 
 	return FReply::Handled();
 }
+TSharedPtr<SWidget> SFoliageEdit::ConstructFoliageMeshContextMenu() const
+{
+	const FFoliageEditCommands& Commands = FFoliageEditCommands::Get();
+	
+	FMenuBuilder MenuBuilder(true, UICommandList);
+	MenuBuilder.AddMenuEntry(Commands.RemoveFoliageType);
+
+	MenuBuilder.AddSubMenu(
+		LOCTEXT( "ReplaceFoliageType", "Replace Foliage Type" ),
+		LOCTEXT ("ReplaceFoliageType_ToolTip", "Replaces selected foliage type with another foliage type asset"),
+		FNewMenuDelegate::CreateSP(this, &SFoliageEdit::FillReplaceFoliageTypeSubmenu));
+	
+	MenuBuilder.AddMenuEntry(Commands.ShowFoliageTypeInCB);
+
+	return MenuBuilder.MakeWidget();
+}
 
 void SFoliageEdit::OnCheckStateChanged_Landscape(ECheckBoxState InState)
 {
@@ -743,4 +732,115 @@ EVisibility SFoliageEdit::GetVisibility_Filters() const
 	return EVisibility::Visible;
 }
 
+TSharedRef<ITableRow> SFoliageEdit::MeshTreeGenerateRow(FFoliageMeshUIInfoPtr Item, const TSharedRef<STableViewBase>& OwnerTable)
+{
+	return SNew(SFoliageEditMeshRow, OwnerTable)
+		.MeshInfo(Item)
+		.FoliageEditMode(FoliageEditMode);
+}
+
+void SFoliageEdit::MeshTreeGetChildren(FFoliageMeshUIInfoPtr Item, TArray<FFoliageMeshUIInfoPtr>& OutChildren)
+{
+	//OutChildren = Item->GetChildren();
+}
+
+void SFoliageEdit::MeshTreeOnSelectionChanged(FFoliageMeshUIInfoPtr Item, ESelectInfo::Type SelectInfo)
+{
+	TArray<UObject*> FoliageMeshes;
+	auto InfoList = MeshTreeWidget->GetSelectedItems();
+
+	for (const auto& Info : InfoList)
+	{
+		FoliageMeshes.Add(Info->Settings);
+	}
+
+	MeshDetailsWidget->SetObjects(FoliageMeshes, true);
+}
+
+void SFoliageEdit::OnRemoveFoliageType()
+{
+	int32 NumInstances = 0;
+	TArray<UFoliageType*> FoliageTypeList;
+	
+	auto MeshUIList = MeshTreeWidget->GetSelectedItems();
+	for (FFoliageMeshUIInfoPtr& MeshUIInfo : MeshUIList)
+	{
+		NumInstances+= MeshUIInfo->InstanceCountTotal;
+		FoliageTypeList.Add(MeshUIInfo->Settings);
+	}
+	
+	bool bProceed = true;
+	if (NumInstances > 0)
+	{
+		FText Message = FText::Format(NSLOCTEXT("UnrealEd", "FoliageMode_DeleteMesh", "Are you sure you want to remove {0} instances?"), FText::AsNumber(NumInstances));
+		bProceed = (FMessageDialog::Open(EAppMsgType::YesNo, Message) == EAppReturnType::Yes);
+	}
+
+	if (bProceed)
+	{
+		FoliageEditMode->RemoveFoliageType(FoliageTypeList.GetData(), FoliageTypeList.Num());
+	}
+}
+
+void SFoliageEdit::FillReplaceFoliageTypeSubmenu(FMenuBuilder& MenuBuilder)
+{
+	FContentBrowserModule& ContentBrowserModule = FModuleManager::Get().LoadModuleChecked<FContentBrowserModule>(TEXT("ContentBrowser"));
+
+	FAssetPickerConfig AssetPickerConfig;
+	AssetPickerConfig.Filter.ClassNames.Add(UFoliageType::StaticClass()->GetFName());
+	AssetPickerConfig.Filter.bRecursiveClasses = true;
+	AssetPickerConfig.OnAssetSelected = FOnAssetSelected::CreateSP(this, &SFoliageEdit::OnReplaceFoliageTypeSelected);
+	AssetPickerConfig.InitialAssetViewType = EAssetViewType::List;
+	AssetPickerConfig.bAllowNullSelection = false;
+	
+	TSharedRef<SWidget> MenuContent = SNew(SBox)
+	.WidthOverride(384)
+	.HeightOverride(500)
+	[
+		ContentBrowserModule.Get().CreateAssetPicker(AssetPickerConfig)
+	];
+	
+	MenuBuilder.AddWidget( MenuContent, FText::GetEmpty(), true);
+}
+
+void SFoliageEdit::OnReplaceFoliageTypeSelected(const FAssetData& AssetData)
+{
+	FSlateApplication::Get().DismissAllMenus();
+	
+	auto MeshUIList = MeshTreeWidget->GetSelectedItems();
+	if (MeshUIList.Num())
+	{
+		UFoliageType* OldSettings = MeshUIList[0]->Settings;
+
+		UFoliageType* NewSettings = Cast<UFoliageType>(AssetData.GetAsset());
+		if (NewSettings && OldSettings != NewSettings)
+		{
+			FoliageEditMode->ReplaceSettingsObject(OldSettings, NewSettings);
+		}
+	}
+}
+
+void SFoliageEdit::OnShowFoliageTypeInCB()
+{
+	TArray<UObject*> SelectedAssets;
+	auto MeshUIList = MeshTreeWidget->GetSelectedItems();
+	for (FFoliageMeshUIInfoPtr  MeshUIPtr : MeshUIList)
+	{
+		UFoliageType* FoliageType = MeshUIPtr->Settings;
+		if (FoliageType->IsAsset())		
+		{
+			SelectedAssets.Add(FoliageType);
+		}
+		else
+		{
+			SelectedAssets.Add(FoliageType->GetStaticMesh());
+		}
+	}
+
+	if (SelectedAssets.Num())
+	{
+		GEditor->SyncBrowserToObjects(SelectedAssets);
+	}
+}
+
 #undef LOCTEXT_NAMESPACE
diff --git a/Engine/Source/Editor/FoliageEdit/Private/SFoliageEdit.h b/Engine/Source/Editor/FoliageEdit/Private/SFoliageEdit.h
--- a/Engine/Source/Editor/FoliageEdit/Private/SFoliageEdit.h
+++ b/Engine/Source/Editor/FoliageEdit/Private/SFoliageEdit.h
@@ -3,10 +3,11 @@
 #pragma once
 
 // Forwards declarations
-class FAssetThumbnail;
-class SFoliageEditMeshDisplayItem;
 struct FFoliageMeshUIInfo;
-class FAssetThumbnailPool;
+class IDetailsView;
+
+typedef TSharedPtr<FFoliageMeshUIInfo> FFoliageMeshUIInfoPtr; //should match typedef in FoliageEdMode.h
+typedef STreeView<FFoliageMeshUIInfoPtr> SFoliageMeshTree;
 
 class SFoliageEdit : public SCompoundWidget
 {
@@ -20,26 +21,18 @@ public:
 
 	~SFoliageEdit();
 
-	/** Creates the thumbnail for the passed in Static Mesh. Used by the MeshListView items. */
-	TSharedPtr<FAssetThumbnail> CreateThumbnail(UStaticMesh* InStaticMesh);
-
 	/** Does a full refresh on the list. */
 	void RefreshFullList();
 
-	/** Adds a static mesh to the list of available meshes for foliage. May be called on by the MeshListView items. */
-	void AddItemToScrollbox(TSharedPtr<FFoliageMeshUIInfo>& InFoliageInfoToAdd);
-
-	/** Removes a static mesh from the list of available meshes for foliage. May be called on by the MeshListView items. */
-	void RemoveItemFromScrollbox(const TSharedPtr<SFoliageEditMeshDisplayItem> InWidgetToRemove);
-
-	void ReplaceItem(const TSharedPtr<SFoliageEditMeshDisplayItem> InDisplayItemToReplaceIn, UStaticMesh* InNewStaticMesh);
-
 	/** Handles adding a new item to the list and refreshing the list in it's entirety. */
 	FReply OnDrop_ListView(const FGeometry& MyGeometry, const FDragDropEvent& DragDropEvent);
 
 	/** Gets FoliageEditMode. Used by the cluster details to notify changes */
 	class FEdModeFoliage* GetFoliageEditMode() const { return FoliageEditMode; }
 
+	/** @return the SWidget containing the context menu */
+	TSharedPtr<SWidget> ConstructFoliageMeshContextMenu() const;
+
 private:
 	/** Clears all the tools selection by setting them to false. */
 	void ClearAllToolSelection();
@@ -99,10 +92,7 @@ private:
 
 	/** Retrieves the Erase Density for the brush. */
 	float GetEraseDensity() const;
-
-	/** Creates the list item widget that displays the instance settings. */
-	TSharedRef<ITableRow> MakeWidgetFromOption(TSharedPtr<FFoliageMeshUIInfo> InItem, const TSharedRef<STableViewBase>& OwnerTable);
-
+	
 	/** Sets the filter settings for if painting will occur on Landscapes. */
 	void OnCheckStateChanged_Landscape(ECheckBoxState InState);
 
@@ -148,31 +138,37 @@ private:
 	/** Checks if the filters should appear. Dependant on the current tool being used. */
 	EVisibility GetVisibility_Filters() const;
 
-	/**
-	 * Checks if a static mesh can be added to the list of Static Meshes available.
-	 *
-	 * @param InStaticMesh		The static mesh to test.
-	 *
-	 * @return	Returns true if the static mesh is not currently in the list, false if it is.
-	 */
-	bool CanAddStaticMesh(const UStaticMesh* const InStaticMesh) const;
-private:
-	/** The list view object for displaying Static Meshes to use for foliage. */
-	TSharedPtr<SListView<TSharedPtr<FFoliageMeshUIInfo>>> MeshListView;
+	/** Generates a row widget for foliage mesh item */
+	TSharedRef<ITableRow> MeshTreeGenerateRow(FFoliageMeshUIInfoPtr Item, const TSharedRef<STableViewBase>& OwnerTable);
 	
-	/** Pool for maintaining and rendering thumbnails */
-	TSharedPtr<FAssetThumbnailPool> AssetThumbnailPool;
+	/** Generates a list of children items for foliage item */
+	void MeshTreeGetChildren(FFoliageMeshUIInfoPtr Item, TArray<FFoliageMeshUIInfoPtr>& OutChildren);
 
-	/** Command list for binding functions for the toolbar. */
-	TSharedPtr<FUICommandList> UICommandList;
+	/** Handler for mesh list view selection changes  */
+	void MeshTreeOnSelectionChanged(FFoliageMeshUIInfoPtr Item, ESelectInfo::Type SelectInfo);
 
-	/** Pointer to the foliage edit mode. */
-	FEdModeFoliage* FoliageEditMode;
+	/** Fills 'Replace' menu command  */
+	void FillReplaceFoliageTypeSubmenu(FMenuBuilder& MenuBuilder);
+	
+	/** Handler for 'Remove' command  */
+	void OnRemoveFoliageType();
+	
+	/** Handler for 'Show in CB' command  */
+	void OnShowFoliageTypeInCB();
 
-	/** Scrollbox for slotting foliage items. */
-	TSharedPtr<SScrollBox> ItemScrollBox;
+	/** Handler for 'Replace' command  */
+	void OnReplaceFoliageTypeSelected(const class FAssetData& AssetData);
 
-	/** List of items currently being displayed. */
-	TArray<TSharedRef<SFoliageEditMeshDisplayItem>> DisplayItemList;
+private:
+	/** Foliage mesh tree widget  */
+	TSharedPtr<SFoliageMeshTree>	MeshTreeWidget;
+	
+	/** Foliage mesh details widget  */
+	TSharedPtr<IDetailsView>		MeshDetailsWidget;
+	
+	/** Command list for binding functions for the toolbar. */
+	TSharedPtr<FUICommandList>		UICommandList;
 
+	/** Pointer to the foliage edit mode. */
+	FEdModeFoliage*					FoliageEditMode;
 };
diff --git a/Engine/Source/Editor/FoliageEdit/Private/SFoliageEdit2.cpp b/Engine/Source/Editor/FoliageEdit/Private/SFoliageEdit2.cpp
deleted file mode 100644
diff --git a/Engine/Source/Editor/FoliageEdit/Private/SFoliageEdit2.h b/Engine/Source/Editor/FoliageEdit/Private/SFoliageEdit2.h
deleted file mode 100644
diff --git a/Engine/Source/Editor/FoliageEdit/Private/SFoliageEditMeshDisplayItem.cpp b/Engine/Source/Editor/FoliageEdit/Private/SFoliageEditMeshDisplayItem.cpp
deleted file mode 100644
diff --git a/Engine/Source/Editor/FoliageEdit/Private/SFoliageEditMeshDisplayItem.h b/Engine/Source/Editor/FoliageEdit/Private/SFoliageEditMeshDisplayItem.h
deleted file mode 100644
diff --git a/Engine/Source/Editor/FoliageEdit/Private/SFoliageEditMeshRow.cpp b/Engine/Source/Editor/FoliageEdit/Private/SFoliageEditMeshRow.cpp
--- a/Engine/Source/Editor/FoliageEdit/Private/SFoliageEditMeshRow.cpp
+++ b/Engine/Source/Editor/FoliageEdit/Private/SFoliageEditMeshRow.cpp
@@ -134,32 +134,6 @@ FReply SFoliageEditMeshRow::OnSave()
 const FSlateBrush* SFoliageEditMeshRow::GetMeshSaveBrush() const
 {
 	return FEditorStyle::GetBrush("Level.SaveIcon16x");
-	
-	//if (LevelModel->IsLoaded())
-	//{
-	//	if (LevelModel->IsLocked())
-	//	{
-	//		return FEditorStyle::GetBrush( "Level.SaveDisabledIcon16x" );
-	//	}
-	//	else
-	//	{
-	//		if (LevelModel->IsDirty())
-	//		{
-	//			return SaveButton->IsHovered() ? FEditorStyle::GetBrush( "Level.SaveModifiedHighlightIcon16x" ) :
-	//												FEditorStyle::GetBrush( "Level.SaveModifiedIcon16x" );
-	//		}
-	//		else
-	//		{
-	//			return SaveButton->IsHovered() ? FEditorStyle::GetBrush( "Level.SaveHighlightIcon16x" ) :
-	//												FEditorStyle::GetBrush( "Level.SaveIcon16x" );
-	//		}
-	//	}								
-	//}
-	//else
-	//{
-	//	return FEditorStyle::GetBrush( "Level.EmptyIcon16x" );
-	//}	
-
 }
 
 FText SFoliageEditMeshRow::GetMeshInstanceCountText() const
diff --git a/Engine/Source/Runtime/Foliage/Private/InstancedFoliage.cpp b/Engine/Source/Runtime/Foliage/Private/InstancedFoliage.cpp
--- a/Engine/Source/Runtime/Foliage/Private/InstancedFoliage.cpp
+++ b/Engine/Source/Runtime/Foliage/Private/InstancedFoliage.cpp
@@ -44,6 +44,8 @@ struct FFoliageCustomVersion
 		CrossLevelBase = 5,
 		// FoliageType for details customization
 		FoliageTypeCustomization = 6,
+		// FoliageType for details customization continued
+		FoliageTypeCustomizationScaling = 7,
 		// -----<new versions can be added above this line>-------------------------------------------------
 		VersionPlusOne,
 		LatestVersion = VersionPlusOne - 1
@@ -253,7 +255,7 @@ UFoliageType::UFoliageType(const FObjectInitializer& ObjectInitializer)
 	Radius = 0.0f;
 	AlignToNormal = true;
 	RandomYaw = true;
-	UniformScale = true;
+	Scaling = EFoliageScaling::Uniform;
 	ScaleX.Min = 1.0f;
 	ScaleY.Min = 1.0f;
 	ScaleZ.Min = 1.0f;
@@ -262,7 +264,8 @@ UFoliageType::UFoliageType(const FObjectInitializer& ObjectInitializer)
 	ScaleZ.Max = 1.0f;
 	AlignMaxAngle = 0.0f;
 	RandomPitchAngle = 0.0f;
-	GroundSlope = 45.0f;
+	GroundSlopeAngle.Min = 0.0f;
+	GroundSlopeAngle.Max = 45.0f;
 	Height.Min = -262144.0f;
 	Height.Max = 262144.0f;
 	ZOffset.Min = 0.0f;
@@ -272,9 +275,6 @@ UFoliageType::UFoliageType(const FObjectInitializer& ObjectInitializer)
 	MinimumLayerWeight = 0.5f;
 	DisplayOrder = 0;
 	IsSelected = false;
-	ShowNothing = false;
-	ShowPaintSettings = true;
-	ShowInstanceSettings = false;
 	ReapplyDensityAmount = 1.0f;
 	CollisionWithWorld = false;
 	CollisionScale = FVector(0.9f, 0.9f, 0.9f);
@@ -330,6 +330,8 @@ UFoliageType::UFoliageType(const FObjectInitializer& ObjectInitializer)
 	HeightMax_DEPRECATED = 262144.0f;
 	ZOffsetMin_DEPRECATED = 0.0f;
 	ZOffsetMax_DEPRECATED = 0.0f;
+	UniformScale_DEPRECATED = true;
+	GroundSlope_DEPRECATED = 45.0f;
 #endif// WITH_EDITORONLY_DATA
 }
 
@@ -346,29 +348,81 @@ void UFoliageType::Serialize(FArchive& Ar)
 	}
 
 #if WITH_EDITORONLY_DATA
-	if (Ar.IsLoading() && Ar.CustomVer(FFoliageCustomVersion::GUID) < FFoliageCustomVersion::FoliageTypeCustomization)
+	if (Ar.IsLoading())
 	{
-		ScaleX.Min = ScaleMinX_DEPRECATED;
-		ScaleX.Max = ScaleMaxX_DEPRECATED;
+		if (Ar.CustomVer(FFoliageCustomVersion::GUID) < FFoliageCustomVersion::FoliageTypeCustomization)
+		{
+			ScaleX.Min = ScaleMinX_DEPRECATED;
+			ScaleX.Max = ScaleMaxX_DEPRECATED;
 
-		ScaleY.Min = ScaleMinY_DEPRECATED;
-		ScaleY.Max = ScaleMaxY_DEPRECATED;
+			ScaleY.Min = ScaleMinY_DEPRECATED;
+			ScaleY.Max = ScaleMaxY_DEPRECATED;
 
-		ScaleZ.Min = ScaleMinZ_DEPRECATED;
-		ScaleZ.Max = ScaleMaxZ_DEPRECATED;
+			ScaleZ.Min = ScaleMinZ_DEPRECATED;
+			ScaleZ.Max = ScaleMaxZ_DEPRECATED;
 
-		Height.Min = HeightMin_DEPRECATED;
-		Height.Max = HeightMax_DEPRECATED;
+			Height.Min = HeightMin_DEPRECATED;
+			Height.Max = HeightMax_DEPRECATED;
 
-		ZOffset.Min = ZOffsetMin_DEPRECATED;
-		ZOffset.Max = ZOffsetMax_DEPRECATED;
+			ZOffset.Min = ZOffsetMin_DEPRECATED;
+			ZOffset.Max = ZOffsetMax_DEPRECATED;
 
-		CullDistance.Min = StartCullDistance_DEPRECATED;
-		CullDistance.Max = EndCullDistance_DEPRECATED;
+			CullDistance.Min = StartCullDistance_DEPRECATED;
+			CullDistance.Max = EndCullDistance_DEPRECATED;
+		}
+		
+		if (Ar.CustomVer(FFoliageCustomVersion::GUID) < FFoliageCustomVersion::FoliageTypeCustomizationScaling)
+		{
+			Scaling = UniformScale_DEPRECATED ? EFoliageScaling::Uniform : EFoliageScaling::Free;
+			
+			GroundSlopeAngle.Min = MinGroundSlope_DEPRECATED;
+			GroundSlopeAngle.Max = GroundSlope_DEPRECATED;
+		}
 	}
 #endif// WITH_EDITORONLY_DATA
 }
 
+FVector UFoliageType::GetRandomScale() const
+{
+	FVector Result(1.0f);
+	float LockRand = 0.0f;
+
+	switch (Scaling)
+	{
+	case EFoliageScaling::Uniform:
+		Result.X = ScaleX.Interpolate(FMath::FRand());
+		Result.Y = Result.X;
+		Result.Z = Result.X;
+		break;
+	
+	case EFoliageScaling::Free:
+		Result.X = ScaleX.Interpolate(FMath::FRand());
+		Result.Y = ScaleY.Interpolate(FMath::FRand());
+		Result.Z = ScaleZ.Interpolate(FMath::FRand());
+		break;
+	
+	case EFoliageScaling::LockXY:
+		LockRand = FMath::FRand();
+		Result.X = ScaleX.Interpolate(LockRand);
+		Result.Y = ScaleY.Interpolate(LockRand);
+		Result.Z = ScaleZ.Interpolate(FMath::FRand());
+		break;
+	
+	case EFoliageScaling::LockXZ:
+		LockRand = FMath::FRand();
+		Result.X = ScaleX.Interpolate(LockRand);
+		Result.Y = ScaleY.Interpolate(FMath::FRand());
+		Result.Z = ScaleZ.Interpolate(LockRand);
+	
+	case EFoliageScaling::LockYZ:
+		LockRand = FMath::FRand();
+		Result.X = ScaleX.Interpolate(FMath::FRand());
+		Result.Y = ScaleY.Interpolate(LockRand);
+		Result.Z = ScaleZ.Interpolate(LockRand);
+	}
+	
+	return Result;
+}
 
 UFoliageType_InstancedStaticMesh::UFoliageType_InstancedStaticMesh(const FObjectInitializer& ObjectInitializer)
 	: Super(ObjectInitializer)
@@ -2150,18 +2204,7 @@ bool FPotentialInstance::PlaceInstance(const UWorld* InWorld, const UFoliageType
 {
 	if (DesiredInstance.PlacementMode != EFoliagePlacementMode::Procedural)
 	{
-		if (Settings->UniformScale)
-		{
-			float Scale = Settings->ScaleX.Interpolate(FMath::FRand());
-			Inst.DrawScale3D = FVector(Scale, Scale, Scale);
-		}
-		else
-		{
-			float LockRand = FMath::FRand();
-			Inst.DrawScale3D.X = Settings->ScaleX.Interpolate(Settings->LockScaleX ? LockRand : FMath::FRand());
-			Inst.DrawScale3D.Y = Settings->ScaleY.Interpolate(Settings->LockScaleY ? LockRand : FMath::FRand());
-			Inst.DrawScale3D.Z = Settings->ScaleZ.Interpolate(Settings->LockScaleZ ? LockRand : FMath::FRand());
-		}
+		Inst.DrawScale3D = Settings->GetRandomScale();
 	}
 	else
 	{
diff --git a/Engine/Source/Runtime/Foliage/Public/FoliageType.h b/Engine/Source/Runtime/Foliage/Public/FoliageType.h
--- a/Engine/Source/Runtime/Foliage/Public/FoliageType.h
+++ b/Engine/Source/Runtime/Foliage/Public/FoliageType.h
@@ -7,13 +7,27 @@
 UENUM()
 enum FoliageVertexColorMask
 {
-	FOLIAGEVERTEXCOLORMASK_Disabled,
-	FOLIAGEVERTEXCOLORMASK_Red,
-	FOLIAGEVERTEXCOLORMASK_Green,
-	FOLIAGEVERTEXCOLORMASK_Blue,
-	FOLIAGEVERTEXCOLORMASK_Alpha,
+	FOLIAGEVERTEXCOLORMASK_Disabled UMETA(DisplayName="Disabled"),
+	FOLIAGEVERTEXCOLORMASK_Red		UMETA(DisplayName="Red"),
+	FOLIAGEVERTEXCOLORMASK_Green	UMETA(DisplayName="Green"),
+	FOLIAGEVERTEXCOLORMASK_Blue		UMETA(DisplayName="Blue"),
+	FOLIAGEVERTEXCOLORMASK_Alpha	UMETA(DisplayName="Alpha"),
 };
 
+UENUM()
+enum class EFoliageScaling : uint8
+{
+	// Foliage instances will have uniform X,Y and Z scales
+	Uniform,
+	// Foliage instances will have random X,Y and Z scales
+	Free,
+	// Locks the X and Y axis scale
+	LockXY,
+	// Locks the X and Z axis scale
+	LockXZ,
+	// Locks the Y and Z axis scale
+	LockYZ
+};
 
 UCLASS(hidecategories = Object, editinlinenew, MinimalAPI, BlueprintType, Blueprintable)
 class UFoliageType : public UObject
@@ -25,70 +39,75 @@
 	UPROPERTY()
 	FGuid UpdateGuid;
 
-	UPROPERTY(EditAnywhere, Category=Painting)
+	/** Foliage instances will be placed at this density, specified in instances per 1000x1000 unit area */
+	UPROPERTY(EditAnywhere, Category=Painting, meta=(DisplayName="Density / 1Kuu", UIMin=0, ClampMin=0))
 	float Density;
 
-	UPROPERTY(EditAnywhere, Category=Painting)
+	/** The minimum distance between foliage instances */
+	UPROPERTY(EditAnywhere, Category=Painting, meta=(UIMin=0, ClampMin=0))
 	float Radius;
 
+	/** Specified foliage instance scaling behavior */
 	UPROPERTY(EditAnywhere, Category=Painting)
+	EFoliageScaling Scaling;
+
+	/** Specifies the range of scale, from minimum to maximum, to apply to a foliage instance's X Scale property */
+	UPROPERTY(EditAnywhere, Category=Painting, meta=(UIMin=0.001))
 	FFloatInterval ScaleX;
 
-	UPROPERTY(EditAnywhere, Category=Painting)
+	/** Specifies the range of scale, from minimum to maximum, to apply to a foliage instance's Y Scale property */
+	UPROPERTY(EditAnywhere, Category=Painting, meta=(UIMin=0.001))
 	FFloatInterval ScaleY;
 
-	UPROPERTY(EditAnywhere, Category=Painting)
+	/** Specifies the range of scale, from minimum to maximum, to apply to a foliage instance's Z Scale property */
+	UPROPERTY(EditAnywhere, Category=Painting, meta=(UIMin=0.001))
 	FFloatInterval ScaleZ;
 
-	UPROPERTY(EditAnywhere, Category=Painting)
-	uint32 UniformScale:1;
+	/** Whether foliage instances should have their angle adjusted away from vertical to match the normal of the surface they're painted on */
+	UPROPERTY(EditAnywhere, Category=Placement)
+	uint32 AlignToNormal:1;
 
-	UPROPERTY(EditAnywhere, Category=Painting)
-	uint32 LockScaleX:1;
-
-	UPROPERTY(EditAnywhere, Category=Painting)
-	uint32 LockScaleY:1;
-
-	UPROPERTY(EditAnywhere, Category=Painting)
-	uint32 LockScaleZ:1;
-
-	UPROPERTY(EditAnywhere, Category=Painting)
+	/** The maximum angle in degrees that foliage instances will be adjusted away from the vertical */
+	UPROPERTY(EditAnywhere, Category=Placement)
 	float AlignMaxAngle;
 
-	UPROPERTY(EditAnywhere, Category=Painting)
+	/** A random pitch adjustment can be applied to each instance, up to the specified angle in degrees, from the original vertical */
+	UPROPERTY(EditAnywhere, Category=Placement)
 	float RandomPitchAngle;
 
-	UPROPERTY(EditAnywhere, Category=Painting)
-	float GroundSlope;
+	/* Foliage instances will only be placed on surfaces sloping in the specified angle range from the horizontal */
+	UPROPERTY(EditAnywhere, Category=Placement, meta=(UIMin=0))
+	FFloatInterval GroundSlopeAngle;
 
-	UPROPERTY(EditAnywhere, Category=Painting)
-	float MinGroundSlope;
-
-	UPROPERTY(EditAnywhere, Category=Painting)
+	/* The valid altitude range where foliage instances will be placed, specified using minimum and maximum world coordinate Z values */
+	UPROPERTY(EditAnywhere, Category=Placement)
 	FFloatInterval Height;
 
-	UPROPERTY(EditAnywhere, Category=Painting)
+	/** If a layer name is specified, painting on landscape will limit the foliage to areas of landscape with the specified layer painted */
+	UPROPERTY(EditAnywhere, AdvancedDisplay, Category=Placement)
 	TArray<FName> LandscapeLayers;
 
 	UPROPERTY()
 	FName LandscapeLayer_DEPRECATED;
 	
-	UPROPERTY(EditAnywhere, Category=Painting)
+	/** Specifies the minimum value above which the landscape layer weight value must be, in order for foliage instances to be placed in a specific area */
+	UPROPERTY(EditAnywhere, AdvancedDisplay, Category=Placement)
 	float MinimumLayerWeight;
 
-	UPROPERTY(EditAnywhere, Category=Painting)
-	uint32 AlignToNormal:1;
-
-	UPROPERTY(EditAnywhere, Category=Painting)
+	/** If selected, foliage instances will have a random yaw rotation around their vertical axis applied */
+	UPROPERTY(EditAnywhere, Category=Placement)
 	uint32 RandomYaw:1;
 
-	UPROPERTY(EditAnywhere, Category=Painting)
+	/** Specifies a range from minimum to maximum of the offset to apply to a foliage instance's Z location */
+	UPROPERTY(EditAnywhere, Category=Painting, meta=(DisplayName="Z Offset"))
 	FFloatInterval ZOffset;
 	
-	UPROPERTY(EditAnywhere, Category=Painting)
+	/* If checked, an overlap test with existing world geometry is performed before each instance is placed */
+	UPROPERTY(EditAnywhere, AdvancedDisplay, Category=Painting)
 	uint32 CollisionWithWorld:1;
 
-	UPROPERTY(EditAnywhere, Category=Painting)
+	/* The foliage instance's collision bounding box will be scaled by the specified amount before performing the overlap check */
+	UPROPERTY(EditAnywhere, AdvancedDisplay, Category=Painting)
 	FVector CollisionScale;
 
 	UPROPERTY()
@@ -98,126 +117,140 @@
 	UPROPERTY()
 	FVector LowBoundOriginRadius;
 
-	UPROPERTY(EditAnywhere, Category=Painting)
+	/** 
+	 *  When painting on static meshes, foliage instance placement can be limited to areas where the static mesh has values in the selected vertex color channel(s). 
+	 *  This allows a static mesh to mask out certain areas to prevent foliage from being placed there
+	 */
+	UPROPERTY(EditAnywhere, AdvancedDisplay, Category=Painting)
 	TEnumAsByte<enum FoliageVertexColorMask> VertexColorMask;
 
-	UPROPERTY(EditAnywhere, Category=Painting)
+	/** Specifies the threshold value above which the static mesh vertex color value must be, in order for foliage instances to be placed in a specific area */
+	UPROPERTY(EditAnywhere, AdvancedDisplay, Category=Painting)
 	float VertexColorMaskThreshold;
 
-	UPROPERTY(EditAnywhere, Category=Painting)
+	/** 
+	 *  When unchecked, foliage instances will be placed only when the vertex color in the specified channel(s) is above the threshold amount. 
+	 *  When checked, the vertex color must be less than the threshold amount 
+	 */
+	UPROPERTY(EditAnywhere, AdvancedDisplay, Category=Painting)
 	uint32 VertexColorMaskInvert:1;
 
-	UPROPERTY()
+	UPROPERTY(EditAnywhere, Category=Painting)
 	float ReapplyDensityAmount;
 
-	UPROPERTY()
+	/** If checked, the density of foliage instances already placed will be adjusted. <1 will remove instances to reduce the density, >1 will place extra instances */
+	UPROPERTY(EditAnywhere, Category=Painting)
 	uint32 ReapplyDensity:1;
 
-	UPROPERTY()
+	/** If checked, foliage instances not meeting the new Radius constraint will be removed */
+	UPROPERTY(EditAnywhere, Category=Painting)
 	uint32 ReapplyRadius:1;
 
-	UPROPERTY()
+	/** If checked, foliage instances will have their normal alignment adjusted by the Reapply tool */
+	UPROPERTY(EditAnywhere, Category=Painting)
 	uint32 ReapplyAlignToNormal:1;
 
-	UPROPERTY()
+	/** If checked, foliage instances will have their yaw adjusted by the Reapply tool */
+	UPROPERTY(EditAnywhere, Category=Painting)
 	uint32 ReapplyRandomYaw:1;
 
-	UPROPERTY()
+	/** If checked, foliage instances will have their X scale adjusted by the Reapply tool */
+	UPROPERTY(EditAnywhere, Category=Painting)
 	uint32 ReapplyScaleX:1;
 
-	UPROPERTY()
+	/** If checked, foliage instances will have their Y scale adjusted by the Reapply tool */
+	UPROPERTY(EditAnywhere, Category=Painting)
 	uint32 ReapplyScaleY:1;
 
-	UPROPERTY()
+	/** If checked, foliage instances will have their Z scale adjusted by the Reapply tool */
+	UPROPERTY(EditAnywhere, Category=Painting)
 	uint32 ReapplyScaleZ:1;
 
-	UPROPERTY()
+	/** If checked, foliage instances will have their pitch adjusted by the Reapply tool */
+	UPROPERTY(EditAnywhere, Category=Painting)
 	uint32 ReapplyRandomPitchAngle:1;
 
-	UPROPERTY()
+	/** If checked, foliage instances not meeting the ground slope condition will be removed by the Reapply too */
+	UPROPERTY(EditAnywhere, Category=Painting)
 	uint32 ReapplyGroundSlope:1;
 
-	UPROPERTY()
+	/* If checked, foliage instances not meeting the valid Z height condition will be removed by the Reapply tool */
+	UPROPERTY(EditAnywhere, Category=Painting)
 	uint32 ReapplyHeight:1;
 
-	UPROPERTY()
+	/* If checked, foliage instances painted on areas that do not have the appopriate landscape layer painted will be removed by the Reapply tool */
+	UPROPERTY(EditAnywhere, Category=Painting)
 	uint32 ReapplyLandscapeLayer:1;
 
-	UPROPERTY()
+	/** If checked, foliage instances will have their Z offset adjusted by the Reapply tool */
+	UPROPERTY(EditAnywhere, Category=Painting)
 	uint32 ReapplyZOffset:1;
 
-	UPROPERTY()
+	/* If checked, foliage instances will have an overlap test with the world reapplied, and overlapping instances will be removed by the Reapply tool */
+	UPROPERTY(EditAnywhere, Category=Painting)
 	uint32 ReapplyCollisionWithWorld:1;
 
-	UPROPERTY()
+	/* If checked, foliage instances no longer matching the vertex color constraint will be removed by the Reapply too */
+	UPROPERTY(EditAnywhere, Category=Painting)
 	uint32 ReapplyVertexColorMask:1;
 
 	/**
 	 * The distance where instances will begin to fade out if using a PerInstanceFadeAmount material node. 0 disables.
 	 * When the entire cluster is beyond this distance, the cluster is completely culled and not rendered at all.
 	 */
-	UPROPERTY(EditAnywhere, Category = General, meta = (Subcategory = "Culling"))
+	UPROPERTY(EditAnywhere, Category=InstanceSettings, meta=(UIMin=0))
 	FInt32Interval CullDistance;
-	
-	UPROPERTY()
-	int32 DisplayOrder;
 
-	UPROPERTY()
-	uint32 IsSelected:1;
-
-	UPROPERTY()
-	uint32 ShowNothing:1;
-
-	UPROPERTY()
-	uint32 ShowPaintSettings:1;
-
-	UPROPERTY()
-	uint32 ShowInstanceSettings:1;
-
 	/** Controls whether the foliage should take part in static lighting/shadowing. If false, mesh will not receive or cast static lighting or shadows, regardless of other settings. */
-	UPROPERTY(EditAnywhere, BlueprintReadOnly, Category = General, meta = (Subcategory = "Lighting"))
+	UPROPERTY(EditAnywhere, BlueprintReadOnly, Category=InstanceSettings)
 	uint32 bEnableStaticLighting : 1;
 
 	/** Controls whether the foliage should cast a shadow or not. */
-	UPROPERTY(EditAnywhere, BlueprintReadOnly, Category = General, meta = (Subcategory = "Lighting"))
+	UPROPERTY(EditAnywhere, BlueprintReadOnly, Category=InstanceSettings)
 	uint32 CastShadow:1;
 
 	/** Controls whether the foliage should inject light into the Light Propagation Volume.  This flag is only used if CastShadow is true. */
-	UPROPERTY(EditAnywhere, BlueprintReadOnly, Category = General, meta = (Subcategory = "Lighting"))
+	UPROPERTY(EditAnywhere, BlueprintReadOnly, Category=InstanceSettings)
 	uint32 bAffectDynamicIndirectLighting:1;
 
 	/** Controls whether the primitive should affect dynamic distance field lighting methods.  This flag is only used if CastShadow is true. */
-	UPROPERTY(EditAnywhere, BlueprintReadOnly, Category = General, meta = (Subcategory = "Lighting", EditCondition = "CastShadow"))
+	UPROPERTY(EditAnywhere, BlueprintReadOnly, Category=InstanceSettings, meta = (EditCondition = "CastShadow"))
 	uint32 bAffectDistanceFieldLighting:1;
 
 	/** Controls whether the foliage should cast shadows in the case of non precomputed shadowing.  This flag is only used if CastShadow is true. */
-	UPROPERTY(EditAnywhere, BlueprintReadOnly, Category = General, meta = (Subcategory = "Lighting"))
+	UPROPERTY(EditAnywhere, BlueprintReadOnly, Category=InstanceSettings)
 	uint32 bCastDynamicShadow:1;
 
 	/** Whether the foliage should cast a static shadow from shadow casting lights.  This flag is only used if CastShadow is true. */
-	UPROPERTY(EditAnywhere, BlueprintReadOnly, Category = General, meta = (Subcategory = "Lighting"))
+	UPROPERTY(EditAnywhere, BlueprintReadOnly, Category=InstanceSettings)
 	uint32 bCastStaticShadow:1;
 
 	/** Whether this foliage should cast dynamic shadows as if it were a two sided material. */
-	UPROPERTY(EditAnywhere, AdvancedDisplay, BlueprintReadOnly, Category = General, meta = (Subcategory = "Lighting"))
+	UPROPERTY(EditAnywhere, AdvancedDisplay, BlueprintReadOnly, Category=InstanceSettings)
 	uint32 bCastShadowAsTwoSided:1;
 
 	/** Whether the foliage receives decals. */
-	UPROPERTY(EditAnywhere, AdvancedDisplay, BlueprintReadOnly, Category = General, meta = (Subcategory = "Rendering"))
+	UPROPERTY(EditAnywhere, AdvancedDisplay, BlueprintReadOnly, Category=InstanceSettings)
 	uint32 bReceivesDecals : 1;
 	
 	/** Whether to override the lightmap resolution defined in the static mesh. */
-	UPROPERTY(BlueprintReadOnly, Category = General, meta = (Subcategory = "Lighting"))
+	UPROPERTY(BlueprintReadOnly, Category=InstanceSettings)
 	uint32 bOverrideLightMapRes:1;
 
 	/** Overrides the lightmap resolution defined in the static mesh */
-	UPROPERTY(EditAnywhere, BlueprintReadOnly, Category = General, meta = (Subcategory = "Lighting", DisplayName = "Light Map Resolution", EditCondition = "bOverrideLightMapRes"))
+	UPROPERTY(EditAnywhere, BlueprintReadOnly, Category=InstanceSettings, meta = (DisplayName = "Light Map Resolution", EditCondition = "bOverrideLightMapRes"))
 	int32 OverriddenLightMapRes;
 
 	/** Custom collision for foliage */
-	UPROPERTY(EditAnywhere, Category = General, meta = (Subcategory = "Collision", HideObjectType = true))
+	UPROPERTY(EditAnywhere, Category=InstanceSettings, meta = (HideObjectType = true))
 	struct FBodyInstance BodyInstance;
 
+	UPROPERTY()
+	int32 DisplayOrder;
+
+	UPROPERTY()
+	uint32 IsSelected:1;
+
 	/* Gets/Sets the mesh associated with this FoliageType */
 	virtual UStaticMesh* GetStaticMesh() const PURE_VIRTUAL(UFoliageType::GetStaticMesh, return nullptr; );
 	virtual void SetStaticMesh(UStaticMesh* InStaticMesh) PURE_VIRTUAL(UFoliageType::SetStaticMesh,);
@@ -303,6 +336,8 @@
 
 	virtual void Serialize(FArchive& Ar) override;
 
+	FOLIAGE_API FVector GetRandomScale() const;
+
 private:
 
 	/** The curve used to interpolate the instance scale.*/
@@ -346,6 +381,24 @@
 	
 	UPROPERTY()
 	int32 EndCullDistance_DEPRECATED;
+
+	UPROPERTY()
+	uint32 UniformScale_DEPRECATED:1;
+
+	UPROPERTY()
+	uint32 LockScaleX_DEPRECATED:1;
+
+	UPROPERTY()
+	uint32 LockScaleY_DEPRECATED:1;
+
+	UPROPERTY()
+	uint32 LockScaleZ_DEPRECATED:1;
+
+	UPROPERTY()
+	float GroundSlope_DEPRECATED;
+	
+	UPROPERTY()
+	float MinGroundSlope_DEPRECATED;
 #endif// WITH_EDITORONLY_DATA
 };
 
