Fixed: 'Replace' foliage mesh on multi-selection

diff --git a/Engine/Source/Editor/FoliageEdit/Private/SFoliageEdit.cpp b/Engine/Source/Editor/FoliageEdit/Private/SFoliageEdit.cpp
--- a/Engine/Source/Editor/FoliageEdit/Private/SFoliageEdit.cpp
+++ b/Engine/Source/Editor/FoliageEdit/Private/SFoliageEdit.cpp
@@ -808,14 +808,17 @@ void SFoliageEdit::OnReplaceFoliageTypeSelected(const FAssetData& AssetData)
 	FSlateApplication::Get().DismissAllMenus();
 	
 	auto MeshUIList = MeshTreeWidget->GetSelectedItems();
-	if (MeshUIList.Num())
+	UFoliageType* NewSettings = Cast<UFoliageType>(AssetData.GetAsset());
+	
+	if (MeshUIList.Num() && NewSettings)
 	{
-		UFoliageType* OldSettings = MeshUIList[0]->Settings;
-
-		UFoliageType* NewSettings = Cast<UFoliageType>(AssetData.GetAsset());
-		if (NewSettings && OldSettings != NewSettings)
+		for (auto& MeshUI : MeshUIList)
 		{
-			FoliageEditMode->ReplaceSettingsObject(OldSettings, NewSettings);
+			UFoliageType* OldSettings = MeshUI->Settings;
+			if (OldSettings != NewSettings)
+			{
+				FoliageEditMode->ReplaceSettingsObject(OldSettings, NewSettings);
+			}
 		}
 	}
 }
