Right click to select all and deselect all in foliage list

diff --git a/Engine/Source/Editor/FoliageEdit/Private/FoliageEdMode.cpp b/Engine/Source/Editor/FoliageEdit/Private/FoliageEdMode.cpp
--- a/Engine/Source/Editor/FoliageEdit/Private/FoliageEdMode.cpp
+++ b/Engine/Source/Editor/FoliageEdit/Private/FoliageEdMode.cpp
@@ -981,6 +981,11 @@ void FEdModeFoliage::SelectInstancesForBrush(UWorld* InWorld, const UFoliageType
 	}
 }
 
+void FEdModeFoliage::SelectInstances(const UFoliageType* Settings, bool bSelect)
+{
+	SelectInstances(GetWorld(), Settings, bSelect);
+}
+
 void FEdModeFoliage::SelectInstances(UWorld* InWorld, bool bSelect)
 {
 	for (auto& FoliageMeshUI : FoliageMeshList)
diff --git a/Engine/Source/Editor/FoliageEdit/Private/FoliageEdMode.h b/Engine/Source/Editor/FoliageEdit/Private/FoliageEdMode.h
--- a/Engine/Source/Editor/FoliageEdit/Private/FoliageEdMode.h
+++ b/Engine/Source/Editor/FoliageEdit/Private/FoliageEdMode.h
@@ -316,6 +316,9 @@ public:
 	/** Save the settings object as an asset */
 	UFoliageType* SaveSettingsObject(UFoliageType* Settings);
 
+	/** Set/Clear selection for foliage instances of specific type  */
+	void SelectInstances(const UFoliageType* Settings, bool bSelect);
+
 	/** Add desired instances. Uses foliage settings to determine location/scale/rotation and whether instances should be ignored */
 	static void AddInstances(UWorld* InWorld, const TArray<FDesiredFoliageInstance>& DesiredInstances);
 
diff --git a/Engine/Source/Editor/FoliageEdit/Private/FoliageEditActions.cpp b/Engine/Source/Editor/FoliageEdit/Private/FoliageEditActions.cpp
--- a/Engine/Source/Editor/FoliageEdit/Private/FoliageEditActions.cpp
+++ b/Engine/Source/Editor/FoliageEdit/Private/FoliageEditActions.cpp
@@ -20,5 +20,7 @@ void FFoliageEditCommands::RegisterCommands()
 
 	UI_COMMAND( RemoveFoliageType, "Remove", "Remove all associated foliage instances", EUserInterfaceActionType::Button, FInputGesture() );
 	UI_COMMAND( ShowFoliageTypeInCB, "Show in Content Browser", "Show asset in Content Browser.", EUserInterfaceActionType::Button, FInputGesture() );
+	UI_COMMAND( SelectAllInstances, "Select All Instances", "Select all instances of foliage type.", EUserInterfaceActionType::Button, FInputGesture());
+	UI_COMMAND( DeselectAllInstances, "Deselect All Instances", "Deselect all instances of foliage type.", EUserInterfaceActionType::Button, FInputGesture());
 }
 
diff --git a/Engine/Source/Editor/FoliageEdit/Private/FoliageEditActions.h b/Engine/Source/Editor/FoliageEdit/Private/FoliageEditActions.h
--- a/Engine/Source/Editor/FoliageEdit/Private/FoliageEditActions.h
+++ b/Engine/Source/Editor/FoliageEdit/Private/FoliageEditActions.h
@@ -40,7 +40,8 @@ public:
 	/** FoliageType commands */
 	TSharedPtr< FUICommandInfo > RemoveFoliageType;
 	TSharedPtr< FUICommandInfo > ShowFoliageTypeInCB;
-
+	TSharedPtr< FUICommandInfo > SelectAllInstances;
+	TSharedPtr< FUICommandInfo > DeselectAllInstances;
 	/**
 	 * Initialize commands
 	 */
diff --git a/Engine/Source/Editor/FoliageEdit/Private/SFoliageEdit.cpp b/Engine/Source/Editor/FoliageEdit/Private/SFoliageEdit.cpp
--- a/Engine/Source/Editor/FoliageEdit/Private/SFoliageEdit.cpp
+++ b/Engine/Source/Editor/FoliageEdit/Private/SFoliageEdit.cpp
@@ -433,6 +433,17 @@ void SFoliageEdit::BindCommands()
 		FCanExecuteAction(),
 		FIsActionChecked());
 
+	UICommandList->MapAction(
+		Commands.SelectAllInstances,
+		FExecuteAction::CreateSP(this, &SFoliageEdit::OnSelectAllInstances),
+		FCanExecuteAction(),
+		FIsActionChecked());
+
+	UICommandList->MapAction(
+		Commands.DeselectAllInstances,
+		FExecuteAction::CreateSP(this, &SFoliageEdit::OnDeselectAllInstances),
+		FCanExecuteAction(),
+		FIsActionChecked());
 }
 
 TSharedRef<SWidget> SFoliageEdit::BuildToolBar()
@@ -646,9 +657,15 @@ TSharedPtr<SWidget> SFoliageEdit::ConstructFoliageMeshContextMenu() const
 		LOCTEXT( "ReplaceFoliageType", "Replace Foliage Type" ),
 		LOCTEXT ("ReplaceFoliageType_ToolTip", "Replaces selected foliage type with another foliage type asset"),
 		FNewMenuDelegate::CreateSP(this, &SFoliageEdit::FillReplaceFoliageTypeSubmenu));
-	
+
 	MenuBuilder.AddMenuEntry(Commands.ShowFoliageTypeInCB);
 
+	if (IsSelectTool() || IsLassoSelectTool())
+	{
+		MenuBuilder.AddMenuEntry(Commands.SelectAllInstances);
+		MenuBuilder.AddMenuEntry(Commands.DeselectAllInstances);
+	}
+
 	return MenuBuilder.MakeWidget();
 }
 
@@ -846,4 +863,26 @@ void SFoliageEdit::OnShowFoliageTypeInCB()
 	}
 }
 
+void SFoliageEdit::OnSelectAllInstances()
+{
+	TArray<UObject*> SelectedAssets;
+	auto MeshUIList = MeshTreeWidget->GetSelectedItems();
+	for (FFoliageMeshUIInfoPtr MeshUIPtr : MeshUIList)
+	{
+		UFoliageType* FoliageType = MeshUIPtr->Settings;
+		FoliageEditMode->SelectInstances(FoliageType, true);
+	}
+}
+
+void SFoliageEdit::OnDeselectAllInstances()
+{
+	TArray<UObject*> SelectedAssets;
+	auto MeshUIList = MeshTreeWidget->GetSelectedItems();
+	for (FFoliageMeshUIInfoPtr MeshUIPtr : MeshUIList)
+	{
+		UFoliageType* FoliageType = MeshUIPtr->Settings;
+		FoliageEditMode->SelectInstances(FoliageType, false);
+	}
+}
+
 #undef LOCTEXT_NAMESPACE
diff --git a/Engine/Source/Editor/FoliageEdit/Private/SFoliageEdit.h b/Engine/Source/Editor/FoliageEdit/Private/SFoliageEdit.h
--- a/Engine/Source/Editor/FoliageEdit/Private/SFoliageEdit.h
+++ b/Engine/Source/Editor/FoliageEdit/Private/SFoliageEdit.h
@@ -159,6 +159,12 @@ private:
 	/** Handler for 'Replace' command  */
 	void OnReplaceFoliageTypeSelected(const class FAssetData& AssetData);
 
+	/** Handler for 'Select All' command  */
+	void OnSelectAllInstances();
+
+	/** Handler for 'Deselect All' command  */
+	void OnDeselectAllInstances();
+
 private:
 	/** Foliage mesh tree widget  */
 	TSharedPtr<SFoliageMeshTree>	MeshTreeWidget;
