Fixed: Editor crashes when erasing foliage with Erase Density set to 2.0
UE-11261

diff --git a/Engine/Source/Editor/FoliageEdit/Private/FoliageEdMode.cpp b/Engine/Source/Editor/FoliageEdit/Private/FoliageEdMode.cpp
--- a/Engine/Source/Editor/FoliageEdit/Private/FoliageEdMode.cpp
+++ b/Engine/Source/Editor/FoliageEdit/Private/FoliageEdMode.cpp
@@ -920,6 +920,11 @@ void FEdModeFoliage::RemoveInstancesForBrush(UWorld* InWorld, const UFoliageType
 		}
 		
 		int32 InstancesToRemove = FMath::RoundToInt((float)(PotentialInstancesToRemove.Num() - DesiredInstanceCount) * Pressure);
+		if (InstancesToRemove <= 0)
+		{
+			continue;
+		}
+		
 		int32 InstancesToKeep = PotentialInstancesToRemove.Num() - InstancesToRemove;
 		if (InstancesToKeep > 0)
 		{
