fix for UE-9947 Crash occurred restoring viewport with foliage in the world

diff --git a/Engine/Source/Editor/FoliageEdit/Private/FoliageEdMode.cpp b/Engine/Source/Editor/FoliageEdit/Private/FoliageEdMode.cpp
--- a/Engine/Source/Editor/FoliageEdit/Private/FoliageEdMode.cpp
+++ b/Engine/Source/Editor/FoliageEdit/Private/FoliageEdMode.cpp
@@ -324,7 +324,7 @@ void FEdModeFoliage::Tick(FEditorViewportClient* ViewportClient, float DeltaTime
 void FEdModeFoliage::FoliageBrushTrace(FEditorViewportClient* ViewportClient, int32 MouseX, int32 MouseY)
 {
 	bBrushTraceValid = false;
-	if (!ViewportClient->IsMovingCamera())
+	if (!ViewportClient->IsMovingCamera() && ViewportClient->IsVisible())
 	{
 		if (UISettings.GetPaintToolSelected() || UISettings.GetReapplyToolSelected() || UISettings.GetLassoSelectToolSelected())
 		{
