Fixed reapply density tool for foliage painting
UE-11710

diff --git a/Engine/Source/Editor/FoliageEdit/Private/FoliageEdMode.cpp b/Engine/Source/Editor/FoliageEdit/Private/FoliageEdMode.cpp
--- a/Engine/Source/Editor/FoliageEdit/Private/FoliageEdMode.cpp
+++ b/Engine/Source/Editor/FoliageEdit/Private/FoliageEdMode.cpp
@@ -1166,19 +1166,22 @@ void FEdModeFoliage::RemoveSelectedInstances(UWorld* InWorld)
 	GEditor->EndTransaction();
 }
 
-void FEdModeFoliage::ReapplyInstancesForBrush(UWorld* InWorld, const UFoliageType* Settings, const FSphere& BrushSphere)
+void FEdModeFoliage::ReapplyInstancesForBrush(UWorld* InWorld, const UFoliageType* Settings, const FSphere& BrushSphere, float Pressure)
 {
+	// Adjust instance density first
+	ReapplyInstancesDensityForBrush(InWorld, Settings, BrushSphere, Pressure);
+		
 	for (FFoliageMeshInfoIterator It(InWorld, Settings); It; ++It)
 	{
 		FFoliageMeshInfo* MeshInfo = (*It);
 		AInstancedFoliageActor* IFA = It.GetActor();
 
-		ReapplyInstancesForBrush(InWorld, IFA, Settings, MeshInfo, BrushSphere);
+		ReapplyInstancesForBrush(InWorld, IFA, Settings, MeshInfo, BrushSphere, Pressure);
 	}
 }
 
 /** Reapply instance settings to exiting instances */
-void FEdModeFoliage::ReapplyInstancesForBrush(UWorld* InWorld, AInstancedFoliageActor* IFA, const UFoliageType* Settings, FFoliageMeshInfo* MeshInfo, const FSphere& BrushSphere)
+void FEdModeFoliage::ReapplyInstancesForBrush(UWorld* InWorld, AInstancedFoliageActor* IFA, const UFoliageType* Settings, FFoliageMeshInfo* MeshInfo, const FSphere& BrushSphere, float Pressure)
 {
 	TArray<int32> ExistingInstances;
 	MeshInfo->GetInstancesInsideSphere(BrushSphere, ExistingInstances);
@@ -1486,31 +1489,29 @@ void FEdModeFoliage::ReapplyInstancesForBrush(UWorld* InWorld, AInstancedFoliage
 
 void FEdModeFoliage::ReapplyInstancesDensityForBrush(UWorld* InWorld, const UFoliageType* Settings, const FSphere& BrushSphere, float Pressure)
 {
-	// TODO: make it work with seamless painting
+	if (Settings->ReapplyDensity && !FMath::IsNearlyEqual(Settings->ReapplyDensityAmount, 1.f))
+	{
+		// Determine number of instances at the start of the brush stroke
+		int32 SnapshotInstanceCount = 0;
+		TArray<const FMeshInfoSnapshot*> SnapshotList;
+		InstanceSnapshot.MultiFindPointer(const_cast<UFoliageType*>(Settings), SnapshotList);
+		for (const auto* Snapshot: SnapshotList)
+		{
+			SnapshotInstanceCount+= Snapshot->CountInstancesInsideSphere(BrushSphere);
+		}
 	
-	//// Adjust instance density
-	//int32 SnapShotInstanceCount = 0;
-	//TArray<FMeshInfoSnapshot*> SnapShotList;
-	//InstanceSnapshot.MultiFindPointer(Settings, SnapShotList);
+		// Determine desired number of instances
+		int32 DesiredInstanceCount = FMath::RoundToInt((float)SnapshotInstanceCount * Settings->ReapplyDensityAmount);
 
-	//if (SnapShotList.Num())
-	//{
-	//	// Use snapshot to determine number of instances at the start of the brush stroke
-	//	for (auto* SnapShot : SnapShotList)
-	//	{
-	//		SnapShotInstanceCount+= SnapShot->CountInstancesInsideSphere(BrushSphere);
-	//	}
-	//			
-	//	int32 NewInstanceCount = FMath::RoundToInt((float)SnapShotInstanceCount * Settings->ReapplyDensityAmount);
-	//	if (Settings->ReapplyDensityAmount > 1.f && NewInstanceCount > Instances.Num())
-	//	{
-	//		AddInstancesForBrush(InWorld, Settings, BrushSphere, );
-	//	}
-	//	else if (Settings->ReapplyDensityAmount < 1.f && NewInstanceCount < Instances.Num())
-	//	{
-	//		RemoveInstancesForBrush(IFA, *MeshInfo, NewInstanceCount, Instances, Pressure);
-	//	}
-	//}
+		if (Settings->ReapplyDensityAmount > 1.f)
+		{
+			AddInstancesForBrush(InWorld, Settings, BrushSphere, DesiredInstanceCount, Pressure);
+		}
+		else if (Settings->ReapplyDensityAmount < 1.f)
+		{
+			RemoveInstancesForBrush(InWorld, Settings, BrushSphere, DesiredInstanceCount, Pressure);
+		}
+	}
 }
 
 void FEdModeFoliage::PreApplyBrush()
@@ -1582,7 +1583,7 @@ void FEdModeFoliage::ApplyBrush(FEditorViewportClient* ViewportClient)
 		else if (UISettings.GetReapplyToolSelected())
 		{
 			// Reapply any settings checked by the user
-			ReapplyInstancesForBrush(World, Settings, BrushSphere);
+			ReapplyInstancesForBrush(World, Settings, BrushSphere, Pressure);
 		}
 		else if (UISettings.GetPaintToolSelected())
 		{
diff --git a/Engine/Source/Editor/FoliageEdit/Private/FoliageEdMode.h b/Engine/Source/Editor/FoliageEdit/Private/FoliageEdMode.h
--- a/Engine/Source/Editor/FoliageEdit/Private/FoliageEdMode.h
+++ b/Engine/Source/Editor/FoliageEdit/Private/FoliageEdMode.h
@@ -131,14 +131,16 @@ public:
 	FMeshInfoSnapshot(FFoliageMeshInfo* MeshInfo)
 		: Hash(*MeshInfo->InstanceHash)
 	{
-		Locations.AddUninitialized(MeshInfo->Instances.Num());
-		for (int32 Idx = 0; Idx < MeshInfo->Instances.Num(); Idx++)
+		int32 NumInstances = MeshInfo->Instances.Num();
+		Locations.Reserve(NumInstances);
+		Locations.AddUninitialized(NumInstances);
+		for (int32 Idx = 0; Idx < NumInstances; Idx++)
 		{
 			Locations[Idx] = MeshInfo->Instances[Idx].Location;
 		}
 	}
 
-	int32 CountInstancesInsideSphere(const FSphere& Sphere)
+	int32 CountInstancesInsideSphere(const FSphere& Sphere) const
 	{
 		int32 Count = 0;
 
@@ -337,8 +339,8 @@ private:
 
 	/** Reapply instance settings to exiting instances */
 	void ReapplyInstancesDensityForBrush(UWorld* InWorld, const UFoliageType* Settings, const FSphere& BrushSphere, float Pressure);
-	void ReapplyInstancesForBrush(UWorld* InWorld, const UFoliageType* Settings, const FSphere& BrushSphere);
-	void ReapplyInstancesForBrush(UWorld* InWorld, AInstancedFoliageActor* IFA, const UFoliageType* Settings, FFoliageMeshInfo* MeshInfo, const FSphere& BrushSphere);
+	void ReapplyInstancesForBrush(UWorld* InWorld, const UFoliageType* Settings, const FSphere& BrushSphere, float Pressure);
+	void ReapplyInstancesForBrush(UWorld* InWorld, AInstancedFoliageActor* IFA, const UFoliageType* Settings, FFoliageMeshInfo* MeshInfo, const FSphere& BrushSphere, float Pressure);
 
 	/** Select instances inside the brush. */
 	void SelectInstancesForBrush(UWorld* InWorld, const UFoliageType* Settings, const FSphere& BrushSphere, bool bSelect);
diff --git a/Engine/Source/Editor/FoliageEdit/Private/FoliageTypePaintingCustomization.cpp b/Engine/Source/Editor/FoliageEdit/Private/FoliageTypePaintingCustomization.cpp
--- a/Engine/Source/Editor/FoliageEdit/Private/FoliageTypePaintingCustomization.cpp
+++ b/Engine/Source/Editor/FoliageEdit/Private/FoliageTypePaintingCustomization.cpp
@@ -77,11 +77,15 @@ void FFoliageTypePaintingCustomization::CustomizeDetails(IDetailLayoutBuilder& D
 	ReapplyCollisionWithWorld = DetailLayoutBuilder.GetProperty(GET_MEMBER_NAME_CHECKED(UFoliageType, ReapplyCollisionWithWorld));
 	ReapplyVertexColorMask = DetailLayoutBuilder.GetProperty(GET_MEMBER_NAME_CHECKED(UFoliageType, ReapplyVertexColorMask));
 	
-	// Density
-	// TODO: Reapply density
+	// Density is hidden during Reapply mode
 	AddFoliageProperty(PaintingCategory, Density, InvalidProperty, 
-		TAttribute<EVisibility>(), 
-		TAttribute<bool>::Create(TAttribute<bool>::FGetter::CreateSP(this, &FFoliageTypePaintingCustomization::IsReapplyPropertyEnabled, ReapplyDensity)));
+		TAttribute<EVisibility>::Create(TAttribute<EVisibility>::FGetter::CreateSP(this, &FFoliageTypePaintingCustomization::GetDensityVisibility)), 
+		TAttribute<bool>());
+	
+	// Reapply DensityAmount is visible only during Reapply mode
+	AddFoliageProperty(PaintingCategory, ReapplyDensityAmount, ReapplyDensity, 
+		TAttribute<EVisibility>::Create(TAttribute<EVisibility>::FGetter::CreateSP(this, &FFoliageTypePaintingCustomization::GetReapplyDensityAmountVisibility)), 
+		TAttribute<bool>());
 
 	// Radius
 	AddFoliageProperty(PaintingCategory, Radius, ReapplyRadius, TAttribute<EVisibility>(), TAttribute<bool>());
@@ -314,4 +318,17 @@ bool FFoliageTypePaintingCustomization::IsReapplyPropertyEnabled(TSharedPtr<IPro
 	return true;
 }
 
+EVisibility FFoliageTypePaintingCustomization::GetDensityVisibility() const
+{
+	// mutually exclusive with ReapplyDensityAmmount
+	return !FoliageEditMode->UISettings.GetReapplyToolSelected() ? EVisibility::Visible : EVisibility::Collapsed;
+}
+
+EVisibility FFoliageTypePaintingCustomization::GetReapplyDensityAmountVisibility() const
+{
+	// mutually exclusive with Density
+	return FoliageEditMode->UISettings.GetReapplyToolSelected() ? EVisibility::Visible : EVisibility::Collapsed;
+}
+
+
 #undef LOCTEXT_NAMESPACE
\ No newline at end of file
diff --git a/Engine/Source/Editor/FoliageEdit/Private/FoliageTypePaintingCustomization.h b/Engine/Source/Editor/FoliageEdit/Private/FoliageTypePaintingCustomization.h
--- a/Engine/Source/Editor/FoliageEdit/Private/FoliageTypePaintingCustomization.h
+++ b/Engine/Source/Editor/FoliageEdit/Private/FoliageTypePaintingCustomization.h
@@ -31,6 +31,9 @@ private:
 	void OnReapplyPropertyStateChanged(ECheckBoxState CheckState, TSharedPtr<IPropertyHandle> ReapplyProperty);
 	bool IsReapplyPropertyEnabled(TSharedPtr<IPropertyHandle> ReapplyProperty) const;
 
+	EVisibility GetDensityVisibility() const;
+	EVisibility GetReapplyDensityAmountVisibility() const;
+
 	IDetailPropertyRow& AddFoliageProperty(
 		IDetailCategoryBuilder& Category,
 		TSharedPtr<IPropertyHandle>& Property, 
diff --git a/Engine/Source/Runtime/Foliage/Public/FoliageType.h b/Engine/Source/Runtime/Foliage/Public/FoliageType.h
--- a/Engine/Source/Runtime/Foliage/Public/FoliageType.h
+++ b/Engine/Source/Runtime/Foliage/Public/FoliageType.h
@@ -140,7 +140,7 @@ class UFoliageType : public UObject
 	UPROPERTY(EditAnywhere, AdvancedDisplay, Category=Painting)
 	uint32 VertexColorMaskInvert:1;
 
-	UPROPERTY(EditAnywhere, Category=Painting)
+	UPROPERTY(EditAnywhere, Category=Painting, meta=(DisplayName="Density Amount", UIMin=0, ClampMin=0))
 	float ReapplyDensityAmount;
 
 	/** If checked, the density of foliage instances already placed will be adjusted. <1 will remove instances to reduce the density, >1 will place extra instances */
diff --git a/Engine/Source/Runtime/Foliage/Public/InstancedFoliage.h b/Engine/Source/Runtime/Foliage/Public/InstancedFoliage.h
--- a/Engine/Source/Runtime/Foliage/Public/InstancedFoliage.h
+++ b/Engine/Source/Runtime/Foliage/Public/InstancedFoliage.h
@@ -239,12 +239,12 @@ private:
 	const int32 HashCellBits;
 	TMap<uint64, TSet<int32>> CellMap;
 
-	uint64 MakeKey(int32 CellX, int32 CellY)
+	uint64 MakeKey(int32 CellX, int32 CellY) const
 	{
 		return ((uint64)(*(uint32*)(&CellX)) << 32) | (*(uint32*)(&CellY) & 0xffffffff);
 	}
 
-	uint64 MakeKey(const FVector& Location)
+	uint64 MakeKey(const FVector& Location) const
 	{
 		return  MakeKey(FMath::FloorToInt(Location.X) >> HashCellBits, FMath::FloorToInt(Location.Y) >> HashCellBits);
 	}
@@ -269,7 +269,7 @@ public:
 		check(RemoveCount == 1);
 	}
 
-	void GetInstancesOverlappingBox(const FBox& InBox, TArray<int32>& OutInstanceIndices)
+	void GetInstancesOverlappingBox(const FBox& InBox, TArray<int32>& OutInstanceIndices) const
 	{
 		int32 MinX = FMath::FloorToInt(InBox.Min.X) >> HashCellBits;
 		int32 MinY = FMath::FloorToInt(InBox.Min.Y) >> HashCellBits;
@@ -290,7 +290,7 @@ public:
 		}
 	}
 
-	TArray<int32> GetInstancesOverlappingBox(const FBox& InBox)
+	TArray<int32> GetInstancesOverlappingBox(const FBox& InBox) const
 	{
 		TArray<int32> Result;
 		GetInstancesOverlappingBox(InBox, Result);
@@ -298,7 +298,7 @@ public:
 	}
 
 #if UE_BUILD_DEBUG
-	void CheckInstanceCount(int32 InCount)
+	void CheckInstanceCount(int32 InCount) const
 	{
 		int32 HashCount = 0;
 		for (const auto& Pair : CellMap)
