[UE-11751] Renamed ProceduralFoliageActor to ProceduralFoliageVolume

diff --git a/Engine/Config/BaseEngine.ini b/Engine/Config/BaseEngine.ini
--- a/Engine/Config/BaseEngine.ini
+++ b/Engine/Config/BaseEngine.ini
@@ -353,6 +353,9 @@ bUseStreamingPause=false
 +ActiveClassRedirects=(OldClassName="MaterialExpressionLandscapeLayerWeight",NewClassName="/Script/Landscape.MaterialExpressionLandscapeLayerWeight")
 +ActiveClassRedirects=(OldClassName="MaterialExpressionLandscapeVisibilityMask",NewClassName="/Script/Landscape.MaterialExpressionLandscapeVisibilityMask")
 
+;Procedural foliage redirects
++ActiveClassRedirects=(OldClassName="ProceduralFoliageActor",NewClassName="ProceduralFoliageVolume")
+
 ;Shortened animation compression algorithm class names
 +ActiveClassRedirects=(OldClassName="AnimationCompressionAlgorithm",NewClassName="AnimCompress")
 +ActiveClassRedirects=(OldClassName="AnimationCompressionAlgorithm_Automatic",NewClassName="AnimCompress_Automatic")
diff --git a/Engine/Source/Editor/FoliageEdit/Private/ActorFactoryProceduralFoliage.cpp b/Engine/Source/Editor/FoliageEdit/Private/ActorFactoryProceduralFoliage.cpp
--- a/Engine/Source/Editor/FoliageEdit/Private/ActorFactoryProceduralFoliage.cpp
+++ b/Engine/Source/Editor/FoliageEdit/Private/ActorFactoryProceduralFoliage.cpp
@@ -7,7 +7,7 @@ ActorFactory.cpp:
 #include "UnrealEd.h"
 #include "ActorFactoryProceduralFoliage.h"
 #include "ProceduralFoliage.h"
-#include "ProceduralFoliageActor.h"
+#include "ProceduralFoliageVolume.h"
 #include "ProceduralFoliageComponent.h"
 
 DEFINE_LOG_CATEGORY_STATIC(LogActorFactory, Log, All);
@@ -21,7 +21,7 @@ UActorFactoryProceduralFoliage::UActorFactoryProceduralFoliage(const FObjectInit
 : Super(ObjectInitializer)
 {
 	DisplayName = LOCTEXT("ProceduralFoliageDisplayName", "ProceduralFoliage");
-	NewActorClass = AProceduralFoliageActor::StaticClass();
+	NewActorClass = AProceduralFoliageVolume::StaticClass();
 	bUseSurfaceOrientation = true;
 }
 
@@ -44,8 +44,8 @@ void UActorFactoryProceduralFoliage::PostSpawnActor(UObject* Asset, AActor* NewA
 	UE_LOG(LogActorFactory, Log, TEXT("Actor Factory created %s"), *ProceduralFoliage->GetName());
 
 	// Change properties
-	AProceduralFoliageActor* PFA = CastChecked<AProceduralFoliageActor>(NewActor);
-	UProceduralFoliageComponent* ProceduralComponent = PFA->ProceduralComponent;
+	AProceduralFoliageVolume* ProceduralFoliageVolume = CastChecked<AProceduralFoliageVolume>(NewActor);
+	UProceduralFoliageComponent* ProceduralComponent = ProceduralFoliageVolume->ProceduralComponent;
 	check(ProceduralComponent);
 
 	ProceduralComponent->UnregisterComponent();
@@ -59,18 +59,20 @@ void UActorFactoryProceduralFoliage::PostSpawnActor(UObject* Asset, AActor* NewA
 UObject* UActorFactoryProceduralFoliage::GetAssetFromActorInstance(AActor* Instance)
 {
 	check(Instance->IsA(NewActorClass));
-	AProceduralFoliageActor* ProceduralFoliage = CastChecked<AProceduralFoliageActor>(Instance);
-	UProceduralFoliageComponent* ProceduralComponent = ProceduralFoliage->ProceduralComponent;
+
+	AProceduralFoliageVolume* ProceduralFoliageVolume = CastChecked<AProceduralFoliageVolume>(Instance);
+	UProceduralFoliageComponent* ProceduralComponent = ProceduralFoliageVolume->ProceduralComponent;
 	check(ProceduralComponent);
+	
 	return ProceduralComponent->ProceduralFoliage;
 }
 
 void UActorFactoryProceduralFoliage::PostCreateBlueprint(UObject* Asset, AActor* CDO)
 {
-	if (Asset != NULL && CDO != NULL)
+	if (Asset != nullptr && CDO != nullptr)
 	{
 		UProceduralFoliage* ProceduralFoliage = CastChecked<UProceduralFoliage>(Asset);
-		AProceduralFoliageActor* PFA = CastChecked<AProceduralFoliageActor>(CDO);
+		AProceduralFoliageVolume* PFA = CastChecked<AProceduralFoliageVolume>(CDO);
 		UProceduralFoliageComponent* ProceduralComponent = PFA->ProceduralComponent;
 		ProceduralComponent->ProceduralFoliage = ProceduralFoliage;
 	}
diff --git a/Engine/Source/Editor/FoliageEdit/Private/FoliageEditModule.cpp b/Engine/Source/Editor/FoliageEdit/Private/FoliageEditModule.cpp
--- a/Engine/Source/Editor/FoliageEdit/Private/FoliageEditModule.cpp
+++ b/Engine/Source/Editor/FoliageEdit/Private/FoliageEditModule.cpp
@@ -14,7 +14,7 @@ const FName FoliageEditAppIdentifier = FName(TEXT("FoliageEdApp"));
 #include "ProceduralFoliageComponentDetails.h"
 #include "ActorFactoryProceduralFoliage.h"
 #include "ComponentVisualizer.h"
-#include "ProceduralFoliageActor.h"
+#include "ProceduralFoliageVolume.h"
 #include "ProceduralFoliageComponent.h"
 
 /**
@@ -46,8 +46,8 @@ public:
 		PropertyEditor.RegisterCustomClassLayout("ProceduralFoliageComponent", FOnGetDetailCustomizationInstance::CreateStatic(&FProceduralFoliageComponentDetails::MakeInstance));
 
 		// Actor Factories
-		auto ProceduralFoliageActorFactory = NewObject<UActorFactoryProceduralFoliage>();
-		GEditor->ActorFactories.Add(ProceduralFoliageActorFactory);
+		auto ProceduralFoliageVolumeFactory = NewObject<UActorFactoryProceduralFoliage>();
+		GEditor->ActorFactories.Add(ProceduralFoliageVolumeFactory);
 
 #if WITH_EDITOR
 		SubscribeEvents();
@@ -83,9 +83,9 @@ public:
 
 	void OnLevelActorDeleted(AActor* Actor)
 	{
-		if (AProceduralFoliageActor* PFA = Cast<AProceduralFoliageActor>(Actor))
+		if (AProceduralFoliageVolume* ProceduralFoliageVolume = Cast<AProceduralFoliageVolume>(Actor))
 		{
-			if (UProceduralFoliageComponent* ProceduralComponent = PFA->ProceduralComponent)
+			if (UProceduralFoliageComponent* ProceduralComponent = ProceduralFoliageVolume->ProceduralComponent)
 			{
 				ProceduralComponent->RemoveProceduralContent();
 			}
diff --git a/Engine/Source/Runtime/Foliage/Private/InstancedFoliage.cpp b/Engine/Source/Runtime/Foliage/Private/InstancedFoliage.cpp
--- a/Engine/Source/Runtime/Foliage/Private/InstancedFoliage.cpp
+++ b/Engine/Source/Runtime/Foliage/Private/InstancedFoliage.cpp
@@ -14,7 +14,7 @@
 #include "Serialization/CustomVersion.h"
 #include "ProceduralFoliageComponent.h"
 #include "ProceduralFoliageBlockingVolume.h"
-#include "ProceduralFoliageActor.h"
+#include "ProceduralFoliageVolume.h"
 
 #define LOCTEXT_NAMESPACE "InstancedFoliage"
 
@@ -2268,8 +2268,8 @@ bool AInstancedFoliageActor::FoliageTrace(const UWorld* InWorld, FHitResult& Out
 		{
 			if (const AProceduralFoliageBlockingVolume* ProceduralFoliageBlockingVolume = Cast<AProceduralFoliageBlockingVolume>(Hit.Actor.Get()))
 			{
-				const AProceduralFoliageActor* ProceduralFoliageActor = ProceduralFoliageBlockingVolume->ProceduralFoliageActor;
-				if (ProceduralFoliageActor == nullptr || ProceduralFoliageActor->ProceduralComponent == nullptr || ProceduralFoliageActor->ProceduralComponent->GetProceduralGuid() == DesiredInstance.ProceduralGuid)
+				const AProceduralFoliageVolume* ProceduralFoliageVolume = ProceduralFoliageBlockingVolume->ProceduralFoliageVolume;
+				if (ProceduralFoliageVolume == nullptr || ProceduralFoliageVolume->ProceduralComponent == nullptr || ProceduralFoliageVolume->ProceduralComponent->GetProceduralGuid() == DesiredInstance.ProceduralGuid)
 				{
 					return false;
 				}
@@ -2278,7 +2278,7 @@ bool AInstancedFoliageActor::FoliageTrace(const UWorld* InWorld, FHitResult& Out
 			{
 				return false;
 			}
-			else if (Cast<AProceduralFoliageActor>(Hit.Actor.Get()))	//we never want to collide with our spawning volume
+			else if (Cast<AProceduralFoliageVolume>(Hit.Actor.Get()))	//we never want to collide with our spawning volume
 			{
 				continue;
 			}
diff --git a/Engine/Source/Runtime/Foliage/Private/ProceduralFoliageActor.cpp b/Engine/Source/Runtime/Foliage/Private/ProceduralFoliageActor.cpp
deleted file mode 100644
diff --git a/Engine/Source/Runtime/Foliage/Private/ProceduralFoliageBlockingVolume.h b/Engine/Source/Runtime/Foliage/Private/ProceduralFoliageBlockingVolume.h
--- a/Engine/Source/Runtime/Foliage/Private/ProceduralFoliageBlockingVolume.h
+++ b/Engine/Source/Runtime/Foliage/Private/ProceduralFoliageBlockingVolume.h
@@ -5,7 +5,7 @@
 #include "GameFramework/Volume.h"
 #include "ProceduralFoliageBlockingVolume.generated.h"
 
-class AProceduralFoliageActor;
+class AProceduralFoliageVolume;
 
 /** An invisible volume used to block ProceduralFoliage instances from being spawned. */
 UCLASS(MinimalAPI)
@@ -14,7 +14,7 @@ class AProceduralFoliageBlockingVolume : public AVolume
 	GENERATED_UCLASS_BODY()
 
 	UPROPERTY(Category = ProceduralFoliage, EditAnywhere)
-	AProceduralFoliageActor* ProceduralFoliageActor;
+	AProceduralFoliageVolume* ProceduralFoliageVolume;
 };
 
 
diff --git a/Engine/Source/Runtime/Foliage/Private/ProceduralFoliageVolume.cpp b/Engine/Source/Runtime/Foliage/Private/ProceduralFoliageVolume.cpp
new file mode 100644
--- /dev/null
+++ b/Engine/Source/Runtime/Foliage/Private/ProceduralFoliageVolume.cpp
@@ -0,0 +1,36 @@
+// Copyright 1998-2015 Epic Games, Inc. All Rights Reserved.
+
+#include "FoliagePrivate.h"
+#include "ProceduralFoliageVolume.h"
+#include "ProceduralFoliageComponent.h"
+#include "ProceduralFoliage.h"
+
+AProceduralFoliageVolume::AProceduralFoliageVolume(const FObjectInitializer& ObjectInitializer)
+: Super(ObjectInitializer)
+{
+	ProceduralComponent = ObjectInitializer.CreateDefaultSubobject<UProceduralFoliageComponent>(this, TEXT("ProceduralFoliageComponent"));
+	ProceduralComponent->SpawningVolume = this;
+
+	if (UBrushComponent* BrushComponent = GetBrushComponent())
+	{
+		BrushComponent->SetCollisionObjectType(ECC_WorldStatic);
+		BrushComponent->SetCollisionResponseToAllChannels(ECR_Ignore);
+		//This is important because the volume overlaps with all procedural foliage, which means during streaming we'll get a huge hitch for UpdateOverlaps
+		BrushComponent->bGenerateOverlapEvents = false;	//We don't care about overlaps.
+	}
+}
+
+#if WITH_EDITOR
+
+bool AProceduralFoliageVolume::GetReferencedContentObjects(TArray<UObject*>& Objects) const
+{
+	Super::GetReferencedContentObjects(Objects);
+
+	if (ProceduralComponent && ProceduralComponent->ProceduralFoliage)
+	{
+		Objects.Add(ProceduralComponent->ProceduralFoliage);
+	}
+	return true;
+}
+
+#endif
\ No newline at end of file
diff --git a/Engine/Source/Runtime/Foliage/Public/ProceduralFoliageActor.h b/Engine/Source/Runtime/Foliage/Public/ProceduralFoliageActor.h
deleted file mode 100644
diff --git a/Engine/Source/Runtime/Foliage/Public/ProceduralFoliageVolume.h b/Engine/Source/Runtime/Foliage/Public/ProceduralFoliageVolume.h
new file mode 100644
--- /dev/null
+++ b/Engine/Source/Runtime/Foliage/Public/ProceduralFoliageVolume.h
@@ -0,0 +1,20 @@
+// Copyright 1998-2015 Epic Games, Inc. All Rights Reserved.
+
+#pragma once
+
+#include "ProceduralFoliageVolume.generated.h"
+
+class UProceduralFoliageComponent;
+
+UCLASS()
+class FOLIAGE_API AProceduralFoliageVolume: public AVolume
+{
+	GENERATED_UCLASS_BODY()
+
+	UPROPERTY(Category = ProceduralFoliageTile, VisibleAnywhere, BlueprintReadOnly)
+	UProceduralFoliageComponent* ProceduralComponent;
+
+#if WITH_EDITOR
+	virtual bool GetReferencedContentObjects(TArray<UObject*>& Objects) const override;
+#endif
+};
\ No newline at end of file
