[UETOOL-265] Undo/Redo works for procedural foliage simulation

diff --git a/Engine/Source/Editor/FoliageEdit/Private/ProceduralFoliageComponentDetails.cpp b/Engine/Source/Editor/FoliageEdit/Private/ProceduralFoliageComponentDetails.cpp
--- a/Engine/Source/Editor/FoliageEdit/Private/ProceduralFoliageComponentDetails.cpp
+++ b/Engine/Source/Editor/FoliageEdit/Private/ProceduralFoliageComponentDetails.cpp
@@ -7,6 +7,7 @@
 #include "ProceduralFoliageComponent.h"
 #include "InstancedFoliage.h"
 #include "FoliageEdMode.h"
+#include "ScopedTransaction.h"
 
 TSharedRef<IDetailCustomization> FProceduralFoliageComponentDetails::MakeInstance()
 {
@@ -73,6 +74,7 @@ FReply FProceduralFoliageComponentDetails::OnResimulateClicked()
 				UniqueFoliageSpawners.Add(Component->FoliageSpawner);
 			}
 
+			FScopedTransaction Transaction(NSLOCTEXT("ProceduralFoliageComponentDetails", "Resimulate_Transaction", "Procedural Foliage Simulation"));
 			TArray <FDesiredFoliageInstance> DesiredFoliageInstances;
 			if (Component->SpawnProceduralContent(DesiredFoliageInstances))
 			{
diff --git a/Engine/Source/Runtime/Foliage/Private/ProceduralFoliageComponent.cpp b/Engine/Source/Runtime/Foliage/Private/ProceduralFoliageComponent.cpp
--- a/Engine/Source/Runtime/Foliage/Private/ProceduralFoliageComponent.cpp
+++ b/Engine/Source/Runtime/Foliage/Private/ProceduralFoliageComponent.cpp
@@ -261,6 +261,7 @@ void UProceduralFoliageComponent::RemoveProceduralContent()
 			AInstancedFoliageActor* IFA = AInstancedFoliageActor::GetInstancedFoliageActorForLevel(Level);
 			if (IFA)
 			{
+				IFA->Modify();
 				IFA->DeleteInstancesForProceduralFoliageComponent(this);
 			}
 		}
