Added checkbox to Foliage UI to toggle all foliage meshes on/off

diff --git a/Engine/Source/Editor/FoliageEdit/Private/SFoliageEdit.cpp b/Engine/Source/Editor/FoliageEdit/Private/SFoliageEdit.cpp
--- a/Engine/Source/Editor/FoliageEdit/Private/SFoliageEdit.cpp
+++ b/Engine/Source/Editor/FoliageEdit/Private/SFoliageEdit.cpp
@@ -283,8 +283,17 @@ void SFoliageEdit::Construct(const FArguments& InArgs)
 			.AutoHeight()
 			[
 				SNew(SHeaderRow)
-
+				+ SHeaderRow::Column(TEXT("ToggleMeshes"))
+				[
+					SNew(SCheckBox)
+					.IsChecked(this, &SFoliageEdit::GetState_AllMeshes)
+					.OnCheckStateChanged(this, &SFoliageEdit::OnCheckStateChanged_AllMeshes)
+				]
+				.HeaderContentPadding(FMargin(0,1,0,1))
+				.FixedWidth(24)
+				
 				+ SHeaderRow::Column(TEXT("Meshes"))
+				.HeaderContentPadding(FMargin(10,1,0,1))
 				.DefaultLabel(LOCTEXT("Meshes", "Meshes"))
 			]
 
@@ -885,4 +894,39 @@ void SFoliageEdit::OnDeselectAllInstances()
 	}
 }
 
+ECheckBoxState SFoliageEdit::GetState_AllMeshes() const
+{
+	bool bHasChecked = false;
+	bool bHasUnchecked = false;
+
+	const auto& MeshUIList = FoliageEditMode->GetFoliageMeshList();
+	for (const auto& MeshUIPtr : MeshUIList)
+	{
+		if (MeshUIPtr->Settings->IsSelected)
+		{
+			bHasChecked = true;
+		}
+		else
+		{
+			bHasUnchecked = true;
+		}
+
+		if (bHasChecked && bHasUnchecked)
+		{
+			return ECheckBoxState::Undetermined; 
+		}
+	}
+
+	return bHasChecked ? ECheckBoxState::Checked : ECheckBoxState::Unchecked;
+}
+
+void SFoliageEdit::OnCheckStateChanged_AllMeshes(ECheckBoxState InState)
+{
+	auto& MeshUIList = FoliageEditMode->GetFoliageMeshList();
+	for (auto& MeshUIPtr : MeshUIList)
+	{
+		MeshUIPtr->Settings->IsSelected = (InState == ECheckBoxState::Checked);
+	}
+}
+
 #undef LOCTEXT_NAMESPACE
diff --git a/Engine/Source/Editor/FoliageEdit/Private/SFoliageEdit.h b/Engine/Source/Editor/FoliageEdit/Private/SFoliageEdit.h
--- a/Engine/Source/Editor/FoliageEdit/Private/SFoliageEdit.h
+++ b/Engine/Source/Editor/FoliageEdit/Private/SFoliageEdit.h
@@ -164,6 +164,10 @@ private:
 
 	/** Handler for 'Deselect All' command  */
 	void OnDeselectAllInstances();
+	
+	/** Toggle all meshes on/off */
+	ECheckBoxState GetState_AllMeshes() const;
+	void OnCheckStateChanged_AllMeshes(ECheckBoxState InState);
 
 private:
 	/** Foliage mesh tree widget  */
