Added 'Drag operation' to FoliageUI mesh selection state, works similar to actors visibility in SceneOutliner

diff --git a/Engine/Source/Editor/FoliageEdit/Private/SFoliageEditMeshRow.cpp b/Engine/Source/Editor/FoliageEdit/Private/SFoliageEditMeshRow.cpp
--- a/Engine/Source/Editor/FoliageEdit/Private/SFoliageEditMeshRow.cpp
+++ b/Engine/Source/Editor/FoliageEdit/Private/SFoliageEditMeshRow.cpp
@@ -6,6 +6,60 @@
 
 #define LOCTEXT_NAMESPACE "FoliageEd_Mode"
 
+/** Drag operation that replicates foliage mesh IsSelected state to other meshes */
+class FFoliageMeshToggleDragDropOp : public FDragDropOperation, public TSharedFromThis<FFoliageMeshToggleDragDropOp>
+{
+public:
+	DRAG_DROP_OPERATOR_TYPE(FFoliageMeshToggleDragDropOp, FDragDropOperation)
+	
+	bool bChecked;
+
+	/** The widget decorator to use */
+	virtual TSharedPtr<SWidget> GetDefaultDecorator() const override
+	{
+		return SNullWidget::NullWidget;
+	}
+	
+	static TSharedRef<FFoliageMeshToggleDragDropOp> New(bool _bChecked)
+	{
+		TSharedRef<FFoliageMeshToggleDragDropOp> Operation = MakeShareable(new FFoliageMeshToggleDragDropOp);
+		Operation->bChecked = _bChecked;
+		Operation->Construct();
+		return Operation;
+	}
+};
+
+/** CheckBox with 'Drag operation' that replicates state to other CheckBoxes */
+class SFoliageMeshToggleCheckBox : public SCheckBox
+{
+	virtual FReply OnMouseButtonDown(const FGeometry& MyGeometry, const FPointerEvent& MouseEvent) override
+	{
+		return SCheckBox::OnMouseButtonDown(MyGeometry, MouseEvent).DetectDrag(SharedThis(this), EKeys::LeftMouseButton);
+	}
+	
+	virtual FReply OnDragDetected(const FGeometry& MyGeometry, const FPointerEvent& MouseEvent) override
+	{
+		if (MouseEvent.IsMouseButtonDown(EKeys::LeftMouseButton))
+		{
+			return FReply::Handled().BeginDragDrop(FFoliageMeshToggleDragDropOp::New(IsChecked()));
+		}
+		else
+		{	
+			return FReply::Unhandled();
+		}
+	}
+
+	virtual void OnDragEnter(const FGeometry& MyGeometry, const FDragDropEvent& DragDropEvent) override
+	{
+		auto ToggleOp = DragDropEvent.GetOperationAs<FFoliageMeshToggleDragDropOp>();
+		if (ToggleOp.IsValid() && IsChecked() != ToggleOp->bChecked)
+		{
+			ToggleCheckedState();
+		}
+	}
+};
+
+/** SFoliageEditMeshRow */
 void SFoliageEditMeshRow::Construct(const FArguments& InArgs, TSharedRef<STableViewBase> InOwnerTableView)
 {
 	MeshInfo = InArgs._MeshInfo;
@@ -20,7 +74,7 @@ TSharedRef<SWidget> SFoliageEditMeshRow::GenerateWidgetForColumn(const FName& Co
 	if (ColumnID == FoliageMeshColumns::ColumnID_ToggleMesh)
 	{
 		TableRowContent =
-			SNew(SCheckBox)
+			SNew(SFoliageMeshToggleCheckBox)
 			.OnCheckStateChanged( this, &SFoliageEditMeshRow::OnMeshCheckStateChanged)
 			.IsChecked( this, &SFoliageEditMeshRow::IsMeshChecked);
 	} 
