[UETOOL-261] - If no instances are spawned following a foliage simulation, a notification is displayed

diff --git a/Engine/Source/Editor/FoliageEdit/Private/ProceduralFoliageComponentDetails.cpp b/Engine/Source/Editor/FoliageEdit/Private/ProceduralFoliageComponentDetails.cpp
--- a/Engine/Source/Editor/FoliageEdit/Private/ProceduralFoliageComponentDetails.cpp
+++ b/Engine/Source/Editor/FoliageEdit/Private/ProceduralFoliageComponentDetails.cpp
@@ -8,6 +8,8 @@
 #include "InstancedFoliage.h"
 #include "FoliageEdMode.h"
 #include "ScopedTransaction.h"
+#include "SNotificationList.h"
+#include "NotificationManager.h"
 
 TSharedRef<IDetailCustomization> FProceduralFoliageComponentDetails::MakeInstance()
 {
@@ -79,6 +81,17 @@ FReply FProceduralFoliageComponentDetails::OnResimulateClicked()
 			if (Component->SpawnProceduralContent(DesiredFoliageInstances))
 			{
 				FEdModeFoliage::AddInstances(Component->GetWorld(), DesiredFoliageInstances);
+
+				// If no instance was actually spawned, inform the user
+				if (!Component->HasSpawnedAnyInstances())
+				{
+					FNotificationInfo Info(NSLOCTEXT("ProceduralFoliageComponentDetails", "NothingSpawned_Notification", "Unable to spawn instances. Ensure a large enough surface exists within the volume."));
+					Info.bUseLargeFont = false;
+					Info.bFireAndForget = true;
+					Info.bUseThrobber = false;
+					Info.bUseSuccessFailIcons = true;
+					FSlateNotificationManager::Get().AddNotification(Info);
+				}
 			}
 		}
 	}
diff --git a/Engine/Source/Runtime/Foliage/Private/InstancedFoliage.cpp b/Engine/Source/Runtime/Foliage/Private/InstancedFoliage.cpp
--- a/Engine/Source/Runtime/Foliage/Private/InstancedFoliage.cpp
+++ b/Engine/Source/Runtime/Foliage/Private/InstancedFoliage.cpp
@@ -1365,6 +1365,26 @@ void AInstancedFoliageActor::DeleteInstancesForProceduralFoliageComponent(const
 	}
 }
 
+bool AInstancedFoliageActor::ContainsInstancesFromProceduralFoliageComponent(const UProceduralFoliageComponent* ProceduralFoliageComponent)
+{
+	const FGuid& ProceduralGuid = ProceduralFoliageComponent->GetProceduralGuid();
+	for (auto& MeshPair : FoliageMeshes)
+	{
+		FFoliageMeshInfo& MeshInfo = *MeshPair.Value;
+		TArray<int32> InstancesToRemove;
+		for (int32 InstanceIdx = 0; InstanceIdx < MeshInfo.Instances.Num(); InstanceIdx++)
+		{
+			if (MeshInfo.Instances[InstanceIdx].ProceduralGuid == ProceduralGuid)
+			{
+				// The procedural component is responsible for an instance
+				return true;
+			}
+		}
+	}
+
+	return false;
+}
+
 void AInstancedFoliageActor::MoveInstancesForComponentToCurrentLevel(UActorComponent* InComponent)
 {
 	if (!HasFoliageAttached(InComponent))
diff --git a/Engine/Source/Runtime/Foliage/Private/ProceduralFoliageComponent.cpp b/Engine/Source/Runtime/Foliage/Private/ProceduralFoliageComponent.cpp
--- a/Engine/Source/Runtime/Foliage/Private/ProceduralFoliageComponent.cpp
+++ b/Engine/Source/Runtime/Foliage/Private/ProceduralFoliageComponent.cpp
@@ -269,4 +269,25 @@ void UProceduralFoliageComponent::RemoveProceduralContent()
 #endif
 }
 
+bool UProceduralFoliageComponent::HasSpawnedAnyInstances()
+{
+	bool bHasSpawnedInstances = false;
+#if WITH_EDITOR
+	UWorld* World = GetWorld();
+
+	for (ULevel* Level : World->GetLevels())
+	{
+		if (Level)
+		{
+			AInstancedFoliageActor* IFA = AInstancedFoliageActor::GetInstancedFoliageActorForLevel(Level);
+			if (IFA)
+			{
+				bHasSpawnedInstances |= IFA->ContainsInstancesFromProceduralFoliageComponent(this);
+			}
+		}
+	}
+#endif
+	return bHasSpawnedInstances;
+}
+
 #undef LOCTEXT_NAMESPACE
\ No newline at end of file
diff --git a/Engine/Source/Runtime/Foliage/Public/InstancedFoliageActor.h b/Engine/Source/Runtime/Foliage/Public/InstancedFoliageActor.h
--- a/Engine/Source/Runtime/Foliage/Public/InstancedFoliageActor.h
+++ b/Engine/Source/Runtime/Foliage/Public/InstancedFoliageActor.h
@@ -125,7 +125,10 @@ public:
 	static FOLIAGE_API void DeleteInstancesForComponent(UWorld* InWorld, UActorComponent* InComponent);
 
 	// Deletes the instances spawned by a procedural component
-	void DeleteInstancesForProceduralFoliageComponent(const UProceduralFoliageComponent* ProceduralComponent);
+	void DeleteInstancesForProceduralFoliageComponent(const UProceduralFoliageComponent* ProceduralFoliageComponent);
+
+	/** @return True if any instances exist that were spawned by the given procedural component */
+	bool ContainsInstancesFromProceduralFoliageComponent(const UProceduralFoliageComponent* ProceduralFoliageComponent);
 
 	// Finds a mesh entry or adds it if it doesn't already exist
 	FOLIAGE_API FFoliageMeshInfo* FindOrAddMesh(UFoliageType* InType);
diff --git a/Engine/Source/Runtime/Foliage/Public/ProceduralFoliageComponent.h b/Engine/Source/Runtime/Foliage/Public/ProceduralFoliageComponent.h
--- a/Engine/Source/Runtime/Foliage/Public/ProceduralFoliageComponent.h
+++ b/Engine/Source/Runtime/Foliage/Public/ProceduralFoliageComponent.h
@@ -30,6 +30,7 @@ class FOLIAGE_API UProceduralFoliageComponent : public UActorComponent
 
 	bool SpawnProceduralContent(TArray<FDesiredFoliageInstance>& OutFoliageInstances);
 	void RemoveProceduralContent();
+	bool HasSpawnedAnyInstances();
 	const FGuid& GetProceduralGuid() const { return ProceduralGuid;  }
 
 	UPROPERTY()
