Added foliage mesh list sorting by name in FoliageUI

diff --git a/Engine/Source/Editor/FoliageEdit/Private/FoliageEdMode.cpp b/Engine/Source/Editor/FoliageEdit/Private/FoliageEdMode.cpp
--- a/Engine/Source/Editor/FoliageEdit/Private/FoliageEdMode.cpp
+++ b/Engine/Source/Editor/FoliageEdit/Private/FoliageEdMode.cpp
@@ -38,6 +38,30 @@
 #define LOCTEXT_NAMESPACE "FoliageEdMode"
 #define FOLIAGE_SNAP_TRACE (10000.f)
 
+//
+// FFoliageMeshUIInfo
+//
+FText FFoliageMeshUIInfo::GetNameText() const
+{
+	const UFoliageType* FoliageType = Settings;
+	if (FoliageType->IsAsset())
+	{
+		return FText::FromString(FoliageType->GetName());
+	}
+	else
+	{
+		UStaticMesh* Mesh = FoliageType->GetStaticMesh();
+		if (Mesh)
+		{
+			return FText::FromString(Mesh->GetName());
+		}
+		else
+		{
+			return LOCTEXT("EmptyMeshName", "[Empty Mesh]");
+		}
+	}
+}
+
 //
 // FFoliageMeshInfo iterator
 //
@@ -129,6 +153,7 @@ FEdModeFoliage::FEdModeFoliage()
 	: FEdMode()
 	, bToolActive(false)
 	, bCanAltDrag(false)
+	, FoliageMeshListSortMode(EColumnSortMode::None)
 {
 	// Load resources and construct brush component
 	UMaterial* BrushMaterial = nullptr;
@@ -2080,16 +2105,32 @@ void FEdModeFoliage::PopulateFoliageMeshList()
 		}
 	}
 
-	auto CompareDisplayOrder = [](const FFoliageMeshUIInfoPtr& A, const FFoliageMeshUIInfoPtr& B)
+	if (FoliageMeshListSortMode != EColumnSortMode::None)
 	{
-		return A->Settings->DisplayOrder < B->Settings->DisplayOrder;
-	};
+		EColumnSortMode::Type SortMode = FoliageMeshListSortMode;
+		auto CompareFoliageType = [SortMode](const FFoliageMeshUIInfoPtr& A, const FFoliageMeshUIInfoPtr& B)
+		{
+			bool CompareResult = (A->GetNameText().CompareToCaseIgnored(B->GetNameText()) <= 0);
+			return SortMode == EColumnSortMode::Ascending ? CompareResult : !CompareResult;
+		};
 
-	FoliageMeshList.Sort(CompareDisplayOrder);
+		FoliageMeshList.Sort(CompareFoliageType);
+	}
 
 	StaticCastSharedPtr<FFoliageEdModeToolkit>(Toolkit)->RefreshFullList();
 }
 
+void FEdModeFoliage::OnFoliageMeshListSortModeChanged(EColumnSortMode::Type InSortMode)
+{
+	FoliageMeshListSortMode = InSortMode;
+	PopulateFoliageMeshList();
+}
+
+EColumnSortMode::Type FEdModeFoliage::GetFoliageMeshListSortMode() const
+{
+	return FoliageMeshListSortMode;
+}
+
 void FEdModeFoliage::OnInstanceCountUpdated(const UFoliageType* FoliageType)
 {
 	int32 EntryIndex = FoliageMeshList.IndexOfByPredicate([&](const FFoliageMeshUIInfoPtr& UIInfoPtr) {
diff --git a/Engine/Source/Editor/FoliageEdit/Private/FoliageEdMode.h b/Engine/Source/Editor/FoliageEdit/Private/FoliageEdMode.h
--- a/Engine/Source/Editor/FoliageEdit/Private/FoliageEdMode.h
+++ b/Engine/Source/Editor/FoliageEdit/Private/FoliageEdMode.h
@@ -118,6 +118,8 @@ struct FFoliageMeshUIInfo
 	{
 		return Settings == Other.Settings;
 	}
+
+	FText GetNameText() const;
 };
 
 typedef TSharedPtr<FFoliageMeshUIInfo> FFoliageMeshUIInfoPtr;
@@ -282,7 +284,13 @@ public:
 	/** Populate mesh with foliage mesh settings found across world */
 	void PopulateFoliageMeshList();
 
-	/**  */
+	/** Handler for mesh list sort mode changes */
+	void OnFoliageMeshListSortModeChanged(EColumnSortMode::Type InSortMode);
+
+	/** Returns foliage mesh list sort mode */
+	EColumnSortMode::Type GetFoliageMeshListSortMode() const;
+
+	/** Handler for foliage mesh instance count changes */
 	void OnInstanceCountUpdated(const UFoliageType* FoliageType);
 
 	/** Whether any of the selected foliage types can be painted into level */
@@ -399,6 +407,7 @@ private:
 	bool bToolActive;
 	bool bCanAltDrag;
 
-	TArray<FFoliageMeshUIInfoPtr> FoliageMeshList;
+	TArray<FFoliageMeshUIInfoPtr>	FoliageMeshList;
+	EColumnSortMode::Type			FoliageMeshListSortMode; 
 };
 
diff --git a/Engine/Source/Editor/FoliageEdit/Private/SFoliageEdit.cpp b/Engine/Source/Editor/FoliageEdit/Private/SFoliageEdit.cpp
--- a/Engine/Source/Editor/FoliageEdit/Private/SFoliageEdit.cpp
+++ b/Engine/Source/Editor/FoliageEdit/Private/SFoliageEdit.cpp
@@ -291,10 +291,12 @@ void SFoliageEdit::Construct(const FArguments& InArgs)
 				]
 				.HeaderContentPadding(FMargin(0,1,0,1))
 				.FixedWidth(24)
-				
+								
 				+ SHeaderRow::Column(TEXT("Meshes"))
 				.HeaderContentPadding(FMargin(10,1,0,1))
-				.DefaultLabel(LOCTEXT("Meshes", "Meshes"))
+				.DefaultLabel(this, &SFoliageEdit::GetMeshesHeaderText)
+				.SortMode(this, &SFoliageEdit::GetMeshColumnSortMode)
+				.OnSort(this, &SFoliageEdit::OnMeshesColumnSortModeChanged)
 			]
 
 			+ SVerticalBox::Slot()
@@ -929,4 +931,20 @@ void SFoliageEdit::OnCheckStateChanged_AllMeshes(ECheckBoxState InState)
 	}
 }
 
+FText SFoliageEdit::GetMeshesHeaderText() const
+{
+	int32 NumMeshes = FoliageEditMode->GetFoliageMeshList().Num();
+	return FText::Format(LOCTEXT("FoliageMeshCount", "Meshes - {0}"), FText::AsNumber(NumMeshes));
+}
+
+EColumnSortMode::Type SFoliageEdit::GetMeshColumnSortMode() const
+{
+	return FoliageEditMode->GetFoliageMeshListSortMode();
+}
+
+void SFoliageEdit::OnMeshesColumnSortModeChanged(EColumnSortPriority::Type InPriority, const FName& InColumnName, EColumnSortMode::Type InSortMode)
+{
+	FoliageEditMode->OnFoliageMeshListSortModeChanged(InSortMode);
+}
+
 #undef LOCTEXT_NAMESPACE
diff --git a/Engine/Source/Editor/FoliageEdit/Private/SFoliageEdit.h b/Engine/Source/Editor/FoliageEdit/Private/SFoliageEdit.h
--- a/Engine/Source/Editor/FoliageEdit/Private/SFoliageEdit.h
+++ b/Engine/Source/Editor/FoliageEdit/Private/SFoliageEdit.h
@@ -169,6 +169,13 @@ private:
 	ECheckBoxState GetState_AllMeshes() const;
 	void OnCheckStateChanged_AllMeshes(ECheckBoxState InState);
 
+	/** Text for foliage meshes list header */
+	FText GetMeshesHeaderText() const;
+	
+	/** Mesh list sorting support */
+	EColumnSortMode::Type GetMeshColumnSortMode() const;
+	void OnMeshesColumnSortModeChanged(EColumnSortPriority::Type InPriority, const FName& InColumnName, EColumnSortMode::Type InSortMode);
+
 private:
 	/** Foliage mesh tree widget  */
 	TSharedPtr<SFoliageMeshTree>	MeshTreeWidget;
diff --git a/Engine/Source/Editor/FoliageEdit/Private/SFoliageEditMeshRow.cpp b/Engine/Source/Editor/FoliageEdit/Private/SFoliageEditMeshRow.cpp
--- a/Engine/Source/Editor/FoliageEdit/Private/SFoliageEditMeshRow.cpp
+++ b/Engine/Source/Editor/FoliageEdit/Private/SFoliageEditMeshRow.cpp
@@ -133,23 +133,7 @@ TSharedRef<SWidget> SFoliageEditMeshRow::GenerateWidgetForColumn(const FName& Co
 
 FText SFoliageEditMeshRow::GetMeshRowText() const
 {
-	const UFoliageType* FoliageType = MeshInfo->Settings;
-	if (FoliageType->IsAsset())
-	{
-		return FText::FromString(FoliageType->GetName());
-	}
-	else
-	{
-		UStaticMesh* Mesh = FoliageType->GetStaticMesh();
-		if (Mesh)
-		{
-			return FText::FromString(Mesh->GetName());
-		}
-		else
-		{
-			return LOCTEXT("EmptyMeshName", "[Empty Mesh]");
-		}
-	}
+	return MeshInfo->GetNameText();
 }
 	
 void SFoliageEditMeshRow::OnMeshCheckStateChanged(const ECheckBoxState NewCheckedState)
