Missed file for 2488908

diff --git a/Engine/Source/Editor/FoliageEdit/Private/SFoliageEditMeshRow.h b/Engine/Source/Editor/FoliageEdit/Private/SFoliageEditMeshRow.h
--- a/Engine/Source/Editor/FoliageEdit/Private/SFoliageEditMeshRow.h
+++ b/Engine/Source/Editor/FoliageEdit/Private/SFoliageEditMeshRow.h
@@ -27,15 +27,27 @@ public:
 	void Construct(const FArguments& InArgs, TSharedRef<STableViewBase> OwnerTableView);
 	virtual TSharedRef<SWidget> GenerateWidgetForColumn(const FName& ColumnName) override;
 
+	// Text for foliage mesh object name field
 	FText GetMeshRowText() const;
 	
+	// Handler for foliage mesh object selection(activation) 
 	void OnMeshCheckStateChanged(const ECheckBoxState NewCheckedState);
+	
+	// Whether foliage mesh object is selected(active)
 	ECheckBoxState IsMeshChecked() const;
 
+	// Whether 'Save' button for foliage mesh should be enabled
+	// Always enabled for private object, will save as shared asset
+	// On shared foliage mesh objects enabled only when asset is dirty
 	bool IsSaveEnabled() const;
+
+	// Handle 'Save' foliage mesh object
 	FReply OnSave();
+	
+	// Icon for 'Save' button
 	const FSlateBrush* GetMeshSaveBrush() const;
 
+	// Instance count text for foliage mesh
 	FText GetMeshInstanceCountText() const;
 	
 private:
