Fixed: Dragging static meshes to Foliage tab in the Mode window cannot be undone
UE-12914

diff --git a/Engine/Source/Editor/FoliageEdit/Private/FoliageEdMode.cpp b/Engine/Source/Editor/FoliageEdit/Private/FoliageEdMode.cpp
--- a/Engine/Source/Editor/FoliageEdit/Private/FoliageEdMode.cpp
+++ b/Engine/Source/Editor/FoliageEdit/Private/FoliageEdMode.cpp
@@ -2187,11 +2187,13 @@ bool FEdModeFoliage::CanPaint(const UFoliageType* FoliageType, const ULevel* InL
 	return false;
 }
 
-UFoliageType* FEdModeFoliage::AddFoliageAsset(UObject* InAssset)
+UFoliageType* FEdModeFoliage::AddFoliageAsset(UObject* InAsset)
 {
 	UFoliageType* FoliageType = nullptr;
+
+	const FScopedTransaction Transaction(NSLOCTEXT("UnrealEd", "FoliageMode_AddMeshTransaction", "Foliage Editing: Add Mesh"));
 	
-	UStaticMesh* StaticMesh = Cast<UStaticMesh>(InAssset);
+	UStaticMesh* StaticMesh = Cast<UStaticMesh>(InAsset);
 	if (StaticMesh)
 	{
 		AInstancedFoliageActor* IFA = AInstancedFoliageActor::GetInstancedFoliageActorForCurrentLevel(GetWorld(), true);
@@ -2203,7 +2205,7 @@ UFoliageType* FEdModeFoliage::AddFoliageAsset(UObject* InAssset)
 	}
 	else 
 	{
-		FoliageType = Cast<UFoliageType>(InAssset);
+		FoliageType = Cast<UFoliageType>(InAsset);
 		if (FoliageType)
 		{
 			AInstancedFoliageActor* IFA = AInstancedFoliageActor::GetInstancedFoliageActorForCurrentLevel(GetWorld(), true);
diff --git a/Engine/Source/Editor/FoliageEdit/Private/FoliageEdMode.h b/Engine/Source/Editor/FoliageEdit/Private/FoliageEdMode.h
--- a/Engine/Source/Editor/FoliageEdit/Private/FoliageEdMode.h
+++ b/Engine/Source/Editor/FoliageEdit/Private/FoliageEdMode.h
@@ -300,7 +300,7 @@ public:
 	static bool CanPaint(const UFoliageType* FoliageType, const ULevel* InLevel);
 
 	/** Add a new asset (FoliageType or StaticMesh) */
-	UFoliageType* AddFoliageAsset(UObject* InAssset);
+	UFoliageType* AddFoliageAsset(UObject* InAsset);
 	
 	/** Remove a list of Foliage types */
 	bool RemoveFoliageType(UFoliageType** FoliageTypes, int32 Num);
diff --git a/Engine/Source/Editor/FoliageEdit/Private/SFoliageEdit.cpp b/Engine/Source/Editor/FoliageEdit/Private/SFoliageEdit.cpp
--- a/Engine/Source/Editor/FoliageEdit/Private/SFoliageEdit.cpp
+++ b/Engine/Source/Editor/FoliageEdit/Private/SFoliageEdit.cpp
@@ -641,20 +641,21 @@ EVisibility SFoliageEdit::GetVisibility_NonEmptyList() const
 FReply SFoliageEdit::OnDrop_ListView(const FGeometry& MyGeometry, const FDragDropEvent& DragDropEvent)
 {
 	TArray<FAssetData> DroppedAssetData = AssetUtil::ExtractAssetDataFromDrag(DragDropEvent);
-
-	for (int32 AssetIdx = 0; AssetIdx < DroppedAssetData.Num(); ++AssetIdx)
 	{
-		GWarn->BeginSlowTask(LOCTEXT("OnDrop_LoadPackage", "Fully Loading Package For Drop"), true, false);
-		UObject* Asset = DroppedAssetData[AssetIdx].GetAsset();
-		GWarn->EndSlowTask();
+		const FScopedTransaction Transaction(NSLOCTEXT("UnrealEd", "FoliageMode_AddMeshTransaction", "Foliage Editing: Add Mesh"));
 
-		UFoliageType* FoliageType = FoliageEditMode->AddFoliageAsset(Asset);
-		if (FoliageType)
+		for (int32 AssetIdx = 0; AssetIdx < DroppedAssetData.Num(); ++AssetIdx)
 		{
-			MeshTreeWidget->RequestTreeRefresh();
+			GWarn->BeginSlowTask(LOCTEXT("OnDrop_LoadPackage", "Fully Loading Package For Drop"), true, false);
+			UObject* Asset = DroppedAssetData[AssetIdx].GetAsset();
+			GWarn->EndSlowTask();
+
+			FoliageEditMode->AddFoliageAsset(Asset);
 		}
 	}
 
+	MeshTreeWidget->RequestTreeRefresh();
+
 	return FReply::Handled();
 }
 TSharedPtr<SWidget> SFoliageEdit::ConstructFoliageMeshContextMenu() const
diff --git a/Engine/Source/Runtime/Foliage/Private/InstancedFoliage.cpp b/Engine/Source/Runtime/Foliage/Private/InstancedFoliage.cpp
--- a/Engine/Source/Runtime/Foliage/Private/InstancedFoliage.cpp
+++ b/Engine/Source/Runtime/Foliage/Private/InstancedFoliage.cpp
@@ -1626,7 +1626,7 @@ UFoliageType* AInstancedFoliageActor::AddFoliageType(const UFoliageType* InType,
 		auto* MeshInfoPtr = FoliageMeshes.Find(FoliageType);
 		if (!MeshInfoPtr)
 		{
-			MarkPackageDirty();
+			Modify();
 			MeshInfo = &FoliageMeshes.Add(FoliageType).Get();
 		}
 		else
@@ -1641,8 +1641,8 @@ UFoliageType* AInstancedFoliageActor::AddFoliageType(const UFoliageType* InType,
 		FoliageType = GetSettingsForMesh(FoliageType->GetStaticMesh(), &MeshInfo);
 		if (FoliageType == nullptr)
 		{
+			Modify();
 			FoliageType = DuplicateObject<UFoliageType>(InType, this);
-			MarkPackageDirty();
 			MeshInfo = &FoliageMeshes.Add(FoliageType).Get();
 		}
 	}
@@ -1690,7 +1690,7 @@ FFoliageMeshInfo* AInstancedFoliageActor::AddMesh(UFoliageType* InType)
 {
 	check(FoliageMeshes.Find(InType) == nullptr);
 
-	MarkPackageDirty();
+	Modify();
 
 	if (InType->DisplayOrder == 0)
 	{
