fixing non-unity

diff --git a/Engine/Source/Editor/FoliageEdit/Private/SFoliageEdit.cpp b/Engine/Source/Editor/FoliageEdit/Private/SFoliageEdit.cpp
--- a/Engine/Source/Editor/FoliageEdit/Private/SFoliageEdit.cpp
+++ b/Engine/Source/Editor/FoliageEdit/Private/SFoliageEdit.cpp
@@ -10,6 +10,7 @@
 #include "FoliageEditActions.h"
 #include "Editor/UnrealEd/Public/AssetThumbnail.h"
 #include "Editor/UnrealEd/Public/DragAndDrop/AssetDragDropOp.h"
+#include "Editor/UnrealEd/Public/ScopedTransaction.h"
 #include "Editor/ContentBrowser/Public/ContentBrowserModule.h"
 #include "Editor/UnrealEd/Public/AssetSelection.h"
 #include "Editor/IntroTutorials/Public/IIntroTutorials.h"
