Case-specific tooltips for Resumulate button

diff --git a/Engine/Source/Editor/FoliageEdit/Private/ProceduralFoliageComponentDetails.cpp b/Engine/Source/Editor/FoliageEdit/Private/ProceduralFoliageComponentDetails.cpp
--- a/Engine/Source/Editor/FoliageEdit/Private/ProceduralFoliageComponentDetails.cpp
+++ b/Engine/Source/Editor/FoliageEdit/Private/ProceduralFoliageComponentDetails.cpp
@@ -11,19 +11,20 @@
 #include "SNotificationList.h"
 #include "NotificationManager.h"
 
+#define LOCTEXT_NAMESPACE "ProceduralFoliageComponentDetails"
+
 TSharedRef<IDetailCustomization> FProceduralFoliageComponentDetails::MakeInstance()
 {
 	return MakeShareable( new FProceduralFoliageComponentDetails() );
 
 }
 
-
 void FProceduralFoliageComponentDetails::CustomizeDetails( IDetailLayoutBuilder& DetailBuilder )
 {
 	const FName ProceduralFoliageCategoryName("ProceduralFoliage");
 	IDetailCategoryBuilder& ProceduralFoliageCategory = DetailBuilder.EditCategory(ProceduralFoliageCategoryName);
 
-	const FText ResimulateText = NSLOCTEXT("ProceduralFoliageComponentDetails","ResimulateButtonText", "Resimulate" );
+	const FText ResimulateText = LOCTEXT("ResimulateButtonText", "Resimulate" );
 
 	TArray< TWeakObjectPtr<UObject> > ObjectsBeingCustomized;
 	DetailBuilder.GetObjectsBeingCustomized( ObjectsBeingCustomized );
@@ -54,7 +55,7 @@ void FProceduralFoliageComponentDetails::CustomizeDetails( IDetailLayoutBuilder&
 	[
 		SNew(SButton)
 		.OnClicked( this, &FProceduralFoliageComponentDetails::OnResimulateClicked )
-		.ToolTipText( NSLOCTEXT("ProceduralFoliageComponentDetails","ResimulateButton_Tooltip", "Runs the procedural foliage spawner simulation. Replaces any existing instances spawned by a previous simulation." ) )
+		.ToolTipText( this, &FProceduralFoliageComponentDetails::GetResimulateTooltipText )
 		.IsEnabled( this, &FProceduralFoliageComponentDetails::IsResimulateEnabled )
 		[
 			SNew( STextBlock )
@@ -76,16 +77,16 @@ FReply FProceduralFoliageComponentDetails::OnResimulateClicked()
 				UniqueFoliageSpawners.Add(Component->FoliageSpawner);
 			}
 
-			FScopedTransaction Transaction(NSLOCTEXT("ProceduralFoliageComponentDetails", "Resimulate_Transaction", "Procedural Foliage Simulation"));
+			FScopedTransaction Transaction(LOCTEXT("Resimulate_Transaction", "Procedural Foliage Simulation"));
 			TArray <FDesiredFoliageInstance> DesiredFoliageInstances;
 			if (Component->GenerateProceduralContent(DesiredFoliageInstances))
 			{
 				FEdModeFoliage::AddInstances(Component->GetWorld(), DesiredFoliageInstances);
 
-				// If no instance was actually spawned, inform the user
+				// If no instances were spawned, inform the user
 				if (!Component->HasSpawnedAnyInstances())
 				{
-					FNotificationInfo Info(NSLOCTEXT("ProceduralFoliageComponentDetails", "NothingSpawned_Notification", "Unable to spawn instances. Ensure a large enough surface exists within the volume."));
+					FNotificationInfo Info(LOCTEXT("NothingSpawned_Notification", "Unable to spawn instances. Ensure a large enough surface exists within the volume."));
 					Info.bUseLargeFont = false;
 					Info.bFireAndForget = true;
 					Info.bUseThrobber = false;
@@ -102,11 +103,42 @@ bool FProceduralFoliageComponentDetails::IsResimulateEnabled() const
 {
 	for(const TWeakObjectPtr<UProceduralFoliageComponent>& Component : SelectedComponents)
 	{
-		if(Component.IsValid() && Component->FoliageSpawner && Component->FoliageSpawner->GetTypes().Num() > 0)
+		if(Component.IsValid() && Component->FoliageSpawner && Component->FoliageSpawner->GetFoliageTypes().Num() > 0)
 		{
 			return true;
 		}
 	}
 
 	return false;
-}
\ No newline at end of file
+}
+
+FText FProceduralFoliageComponentDetails::GetResimulateTooltipText() const
+{
+	FText TooltipText;
+
+	for (const TWeakObjectPtr<UProceduralFoliageComponent>& Component : SelectedComponents)
+	{
+		if (Component.IsValid())
+		{
+			if (!Component->FoliageSpawner)
+			{
+				TooltipText = LOCTEXT("Resimulate_Tooltip_NeedSpawner", "Cannot generate foliage: Assign a Procedural Foliage Spawner to run the procedural foliage simulation");
+				break;
+			}
+			else if (Component->FoliageSpawner->GetFoliageTypes().Num() == 0)
+			{
+				TooltipText = LOCTEXT("Resimulate_Tooltip_EmptySpawner", "Cannot generate foliage: The assigned Procedural Foliage Spawner does not contain any foliage types to spawn.");
+				break;
+			}
+		}
+	}
+
+	if (TooltipText.IsEmpty())
+	{
+		TooltipText = LOCTEXT("Resimulate_Tooltip", "Runs the procedural foliage spawner simulation. Replaces any existing instances spawned by a previous simulation.");
+	}
+
+	return TooltipText;
+}
+
+#undef LOCTEXT_NAMESPACE
\ No newline at end of file
diff --git a/Engine/Source/Editor/FoliageEdit/Private/ProceduralFoliageComponentDetails.h b/Engine/Source/Editor/FoliageEdit/Private/ProceduralFoliageComponentDetails.h
--- a/Engine/Source/Editor/FoliageEdit/Private/ProceduralFoliageComponentDetails.h
+++ b/Engine/Source/Editor/FoliageEdit/Private/ProceduralFoliageComponentDetails.h
@@ -15,6 +15,7 @@ public:
 private:
 	FReply OnResimulateClicked();
 	bool IsResimulateEnabled() const;
+	FText GetResimulateTooltipText() const;
 private:
 	TArray< TWeakObjectPtr<class UProceduralFoliageComponent> > SelectedComponents;
 
