General, basic improvements to the Foliage mode UI
- Added "+Add Foliage Type" button to allow foliage type assets to be created directly
- Adjusted sizing of elements (font & spinboxes) to match details panel & save space
- Changed "ReapplyDensityAmount" to "DensityAdjustmentFactor" to more properly describe its function

- Foliage Type assets can be created in CB
- Overhauled foliage type details customizations to use metadata specifiers for hiding and reapplying
- Updated non-painting FoliageType details customization to display properties properly

diff --git a/Engine/Source/Editor/EditorStyle/Private/SlateEditorStyle.cpp b/Engine/Source/Editor/EditorStyle/Private/SlateEditorStyle.cpp
--- a/Engine/Source/Editor/EditorStyle/Private/SlateEditorStyle.cpp
+++ b/Engine/Source/Editor/EditorStyle/Private/SlateEditorStyle.cpp
@@ -1635,6 +1635,13 @@ void FSlateEditorStyle::FStyle::SetupGeneralStyles()
 		Set( "FoliageEditMode.SelectionBackground", new IMAGE_BRUSH( "Icons/FoliageEditMode/FoliageEditMode_SelectionBackground", Icon32x32 ) );
 		Set( "FoliageEditMode.ItemBackground", new IMAGE_BRUSH( "Icons/FoliageEditMode/FoliageEditMode_Background", Icon64x64 ) );
 		Set( "FoliageEditMode.BubbleBorder", new BOX_BRUSH( "Icons/FoliageEditMode/FoliageEditMode_BubbleBorder", FMargin(8/32.0f) ) );
+
+		Set("FoliageEditMode.AddFoliageType.Text", FTextBlockStyle(NormalText)
+			.SetFont(TTF_CORE_FONT("Fonts/Roboto-Bold", 10))
+			.SetColorAndOpacity(FLinearColor(1.0f, 1.0f, 1.0f))
+			.SetHighlightColor(FLinearColor(1.0f, 1.0f, 1.0f))
+			.SetShadowOffset(FVector2D(1, 1))
+			.SetShadowColorAndOpacity(FLinearColor(0, 0, 0, 0.9f)));
 	}
 #endif // WITH_EDITOR || IS_PROGRAM
 
diff --git a/Engine/Source/Editor/FoliageEdit/Private/FoliageEdMode.cpp b/Engine/Source/Editor/FoliageEdit/Private/FoliageEdMode.cpp
--- a/Engine/Source/Editor/FoliageEdit/Private/FoliageEdMode.cpp
+++ b/Engine/Source/Editor/FoliageEdit/Private/FoliageEdMode.cpp
@@ -301,6 +301,8 @@ void FEdModeFoliage::NotifyToolChanged()
 	{
 		ApplySelectionToComponents(GetWorld(), false);
 	}
+
+	OnToolChanged.Broadcast();
 }
 
 bool FEdModeFoliage::DisallowMouseDeltaTracking() const
@@ -1317,31 +1319,34 @@ void FEdModeFoliage::ReapplyInstancesForBrush(UWorld* InWorld, AInstancedFoliage
 			}
 
 			// Reapply scale
-			FVector NewScale = Settings->GetRandomScale();
-			
-			if (Settings->ReapplyScaleX)
+			if (Settings->ReapplyScaling)
 			{
-				if (Settings->Scaling == EFoliageScaling::Uniform)
+				FVector NewScale = Settings->GetRandomScale();
+
+				if (Settings->ReapplyScaleX)
 				{
-					Instance.DrawScale3D = NewScale;
+					if (Settings->Scaling == EFoliageScaling::Uniform)
+					{
+						Instance.DrawScale3D = NewScale;
+					}
+					else
+					{
+						Instance.DrawScale3D.X = NewScale.X;
+					}
+					bUpdated = true;
 				}
-				else
+
+				if (Settings->ReapplyScaleY)
 				{
-					Instance.DrawScale3D.X = NewScale.X;
+					Instance.DrawScale3D.Y = NewScale.Y;
+					bUpdated = true;
 				}
-				bUpdated = true;
-			}
 
-			if (Settings->ReapplyScaleY)
-			{
-				Instance.DrawScale3D.Y = NewScale.Y;
-				bUpdated = true;
-			}
-
-			if (Settings->ReapplyScaleZ)
-			{
-				Instance.DrawScale3D.Z = NewScale.Z;
-				bUpdated = true;
+				if (Settings->ReapplyScaleZ)
+				{
+					Instance.DrawScale3D.Z = NewScale.Z;
+					bUpdated = true;
+				}
 			}
 
 			// Reapply ZOffset
@@ -1352,7 +1357,7 @@ void FEdModeFoliage::ReapplyInstancesForBrush(UWorld* InWorld, AInstancedFoliage
 			}
 
 			// Find a ground normal for either normal or ground slope check.
-			if (bReapplyNormal || Settings->ReapplyGroundSlope || Settings->ReapplyVertexColorMask || (Settings->ReapplyLandscapeLayer && LandscapeLayersValid(Settings)))
+			if (bReapplyNormal || Settings->ReapplyGroundSlope || Settings->ReapplyVertexColorMask || (Settings->ReapplyLandscapeLayers && LandscapeLayersValid(Settings)))
 			{
 				FHitResult Hit;
 				static const FName NAME_ReapplyInstancesForBrush = TEXT("ReapplyInstancesForBrush");
@@ -1388,7 +1393,7 @@ void FEdModeFoliage::ReapplyInstancesForBrush(UWorld* InWorld, AInstancedFoliage
 					}
 
 					// Cull instances for the landscape layer
-					if (Settings->ReapplyLandscapeLayer && LandscapeLayersValid(Settings))
+					if (Settings->ReapplyLandscapeLayers && LandscapeLayersValid(Settings))
 					{
 						float HitWeight = 1.f;
 						if (GetMaxHitWeight(Hit.Location, Hit.GetComponent(), Settings, &LandscapeLayerCaches, HitWeight))
@@ -1515,7 +1520,7 @@ void FEdModeFoliage::ReapplyInstancesForBrush(UWorld* InWorld, AInstancedFoliage
 
 void FEdModeFoliage::ReapplyInstancesDensityForBrush(UWorld* InWorld, const UFoliageType* Settings, const FSphere& BrushSphere, float Pressure)
 {
-	if (Settings->ReapplyDensity && !FMath::IsNearlyEqual(Settings->ReapplyDensityAmount, 1.f))
+	if (Settings->ReapplyDensity && !FMath::IsNearlyEqual(Settings->DensityAdjustmentFactor, 1.f))
 	{
 		// Determine number of instances at the start of the brush stroke
 		int32 SnapshotInstanceCount = 0;
@@ -1527,13 +1532,13 @@ void FEdModeFoliage::ReapplyInstancesDensityForBrush(UWorld* InWorld, const UFol
 		}
 	
 		// Determine desired number of instances
-		int32 DesiredInstanceCount = FMath::RoundToInt((float)SnapshotInstanceCount * Settings->ReapplyDensityAmount);
+		int32 DesiredInstanceCount = FMath::RoundToInt((float)SnapshotInstanceCount * Settings->DensityAdjustmentFactor);
 
-		if (Settings->ReapplyDensityAmount > 1.f)
+		if (Settings->DensityAdjustmentFactor > 1.f)
 		{
 			AddInstancesForBrush(InWorld, Settings, BrushSphere, DesiredInstanceCount, Pressure);
 		}
-		else if (Settings->ReapplyDensityAmount < 1.f)
+		else if (Settings->DensityAdjustmentFactor < 1.f)
 		{
 			RemoveInstancesForBrush(InWorld, Settings, BrushSphere, DesiredInstanceCount, Pressure);
 		}
@@ -2191,12 +2196,12 @@ bool FEdModeFoliage::CanPaint(const UFoliageType* FoliageType, const ULevel* InL
 UFoliageType* FEdModeFoliage::AddFoliageAsset(UObject* InAsset)
 {
 	UFoliageType* FoliageType = nullptr;
-
-	const FScopedTransaction Transaction(NSLOCTEXT("UnrealEd", "FoliageMode_AddMeshTransaction", "Foliage Editing: Add Mesh"));
 	
+	const FScopedTransaction Transaction(NSLOCTEXT("UnrealEd", "FoliageMode_AddTypeTransaction", "Add Foliage Type"));
+
 	UStaticMesh* StaticMesh = Cast<UStaticMesh>(InAsset);
 	if (StaticMesh)
-	{
+	{	
 		AInstancedFoliageActor* IFA = AInstancedFoliageActor::GetInstancedFoliageActorForCurrentLevel(GetWorld(), true);
 		FoliageType = IFA->GetSettingsForMesh(StaticMesh);
 		if (!FoliageType)
diff --git a/Engine/Source/Editor/FoliageEdit/Private/FoliageEdMode.h b/Engine/Source/Editor/FoliageEdit/Private/FoliageEdMode.h
--- a/Engine/Source/Editor/FoliageEdit/Private/FoliageEdMode.h
+++ b/Engine/Source/Editor/FoliageEdit/Private/FoliageEdMode.h
@@ -333,6 +333,8 @@ public:
 	static void AddInstances(UWorld* InWorld, const TArray<FDesiredFoliageInstance>& DesiredInstances);
 
 	typedef TMap<FName, TMap<ULandscapeComponent*, TArray<uint8> > > LandscapeLayerCacheData;
+
+	FSimpleMulticastDelegate OnToolChanged;
 private:
 
 	/** Add instances inside the brush to match DesiredInstanceCount */
diff --git a/Engine/Source/Editor/FoliageEdit/Private/FoliageTypeCustomizationHelpers.cpp b/Engine/Source/Editor/FoliageEdit/Private/FoliageTypeCustomizationHelpers.cpp
new file mode 100644
--- /dev/null
+++ b/Engine/Source/Editor/FoliageEdit/Private/FoliageTypeCustomizationHelpers.cpp
@@ -0,0 +1,95 @@
+// Copyright 1998-2015 Epic Games, Inc. All Rights Reserved.
+
+#include "UnrealEd.h"
+#include "FoliageTypeCustomizationHelpers.h"
+
+void FFoliageTypeCustomizationHelpers::ModifyFoliagePropertyRow(IDetailPropertyRow* PropertyRow, const TAttribute<EVisibility>& InVisibility, const TAttribute<bool>& InEnabled)
+{
+	if (PropertyRow)
+	{
+		if (InVisibility.IsSet())
+		{
+			PropertyRow->Visibility(InVisibility);
+		}
+		if (InEnabled.IsSet())
+		{
+			PropertyRow->IsEnabled(InEnabled);
+		}
+	}
+}
+
+void FFoliageTypeCustomizationHelpers::HideFoliageCategory(IDetailLayoutBuilder& DetailLayoutBuilder, FName CategoryName)
+{
+	TArray<TSharedRef<IPropertyHandle>> CategoryProperties;
+	DetailLayoutBuilder.EditCategory(CategoryName).GetDefaultProperties(CategoryProperties, true, true);
+
+	for (auto& PropertyHandle : CategoryProperties)
+	{
+		DetailLayoutBuilder.HideProperty(PropertyHandle);
+	}
+}
+
+void FFoliageTypeCustomizationHelpers::BindHiddenPropertyVisibilityGetter(const TSharedPtr<IPropertyHandle>& PropertyHandle, TAttribute<EVisibility>::FGetter& OutVisibilityGetter)
+{
+	// Determine the conditions for showing this property. 
+	// Note that any special cases can override the visibility attribute by using ModifyFoliageProperty
+	
+	if (PropertyHandle->GetProperty()->IsA<UBoolProperty>())
+	{
+		// If hidden behind a bool, show when that bool is true
+		OutVisibilityGetter.BindLambda([=]
+		{
+			bool bState = 0;
+			auto Result = PropertyHandle->GetValue(bState);
+			return bState && Result == FPropertyAccess::Success ? EVisibility::Visible : EVisibility::Collapsed;
+		});
+	}
+	else if (PropertyHandle->GetProperty()->IsA<UByteProperty>())
+	{
+		// If hidden behind a byte/enum, assume that 0 = disabled and show if nonzero
+		OutVisibilityGetter.BindLambda([=]
+		{
+			uint8 Value = 0;
+			auto Result = PropertyHandle->GetValue(Value);
+			return  Value != 0 && Result == FPropertyAccess::Success ? EVisibility::Visible : EVisibility::Collapsed;
+		});
+	}
+	else
+	{
+		auto ArrayPropHandle = PropertyHandle->AsArray();
+		if (ArrayPropHandle.IsValid())
+		{
+			// If hidden behind an array, show when the array isn't empty
+			OutVisibilityGetter.BindLambda([=]
+			{
+				uint32 NumElements = 0;
+				ArrayPropHandle->GetNumElements(NumElements);
+				return  NumElements > 0 ? EVisibility::Visible : EVisibility::Collapsed;
+			});
+		}
+	}
+}
+
+EVisibility FFoliageTypeCustomizationHelpers::GetScaleAxisVisibility(EAxis::Type Axis, const TSharedPtr<IPropertyHandle> ScalingPropertyHandle)
+{
+	if (ScalingPropertyHandle.IsValid())
+	{
+		uint8 ScalingValue;
+		if (ScalingPropertyHandle->GetValue(ScalingValue) == FPropertyAccess::Success)
+		{
+			switch (Axis)
+			{
+			case EAxis::X:
+				return EVisibility::Visible;
+				break;
+
+			case EAxis::Y:
+			case EAxis::Z:
+				return (ScalingValue == (uint8)EFoliageScaling::Uniform) ? EVisibility::Collapsed : EVisibility::Visible;
+				break;
+			}
+		}
+	}
+
+	return EVisibility::Visible;
+}
\ No newline at end of file
diff --git a/Engine/Source/Editor/FoliageEdit/Private/FoliageTypeCustomizationHelpers.h b/Engine/Source/Editor/FoliageEdit/Private/FoliageTypeCustomizationHelpers.h
new file mode 100644
--- /dev/null
+++ b/Engine/Source/Editor/FoliageEdit/Private/FoliageTypeCustomizationHelpers.h
@@ -0,0 +1,26 @@
+// Copyright 1998-2015 Epic Games, Inc. All Rights Reserved.
+
+#pragma once
+
+#include "Axis.h"
+
+class FEdModeFoliage;
+class IDetailLayoutBuilder;
+class IDetailCategoryBuilder;
+class IDetailPropertyRow;
+
+class FFoliageTypeCustomizationHelpers
+{
+public:
+
+	/** Modifies the visibility and enabled attributes of an existing property row */
+	static void ModifyFoliagePropertyRow(IDetailPropertyRow* PropertyRow, const TAttribute<EVisibility>& InVisibility, const TAttribute<bool>& InEnabled);
+
+	/** Hides all properties in the given category */
+	static void HideFoliageCategory(IDetailLayoutBuilder& DetailLayoutBuilder, FName CategoryName);
+
+	/** Binds the appropriate visibility getter for the hidden property */
+	static void BindHiddenPropertyVisibilityGetter(const TSharedPtr<IPropertyHandle>& PropertyHandle, TAttribute<EVisibility>::FGetter& OutVisibilityGetter);
+
+	static EVisibility GetScaleAxisVisibility(EAxis::Type Axis, const TSharedPtr<IPropertyHandle> ScalingPropertyHandle);
+};
\ No newline at end of file
diff --git a/Engine/Source/Editor/FoliageEdit/Private/FoliageTypeDetails.cpp b/Engine/Source/Editor/FoliageEdit/Private/FoliageTypeDetails.cpp
--- a/Engine/Source/Editor/FoliageEdit/Private/FoliageTypeDetails.cpp
+++ b/Engine/Source/Editor/FoliageEdit/Private/FoliageTypeDetails.cpp
@@ -4,31 +4,43 @@
 #include "FoliageEditModule.h"
 #include "PropertyEditing.h"
 #include "FoliageTypeDetails.h"
+#include "FoliageType.h"
+#include "FoliageTypeCustomizationHelpers.h"
 
 TSharedRef<IDetailCustomization> FFoliageTypeDetails::MakeInstance()
 {
 	return MakeShareable(new FFoliageTypeDetails());
-
 }
 
-static void HideFoliageCategory(IDetailLayoutBuilder& DetailBuilder, FName Category)
+void CustomizePropertyRowVisibility(IDetailLayoutBuilder& LayoutBuilder, const TSharedRef<IPropertyHandle>& PropertyHandle, IDetailPropertyRow& PropertyRow)
 {
-	TArray<TSharedRef<IPropertyHandle>> CategoryProperties;
-	DetailBuilder.EditCategory(Category).GetDefaultProperties(CategoryProperties, true, true);
-
-	for (auto& PropertyHandle : CategoryProperties)
+	// Properties with a HideBehind property specified should only be shown if that property is true, non-zero, or not empty
+	static const FName HideBehindName("HideBehind");
+	if (UProperty* Property = PropertyHandle->GetProperty())
 	{
-		DetailBuilder.HideProperty(PropertyHandle);
+		if (Property->HasMetaData(HideBehindName))
+		{
+			TSharedPtr<IPropertyHandle> HiddenBehindPropertyHandle = LayoutBuilder.GetProperty(*Property->GetMetaData(HideBehindName));
+			if (HiddenBehindPropertyHandle.IsValid() && HiddenBehindPropertyHandle->IsValidHandle())
+			{
+				TAttribute<EVisibility>::FGetter VisibilityGetter;
+				FFoliageTypeCustomizationHelpers::BindHiddenPropertyVisibilityGetter(HiddenBehindPropertyHandle, VisibilityGetter);
+
+				PropertyRow.Visibility(TAttribute<EVisibility>::Create(VisibilityGetter));
+			}
+		}
 	}
 }
 
-void AddSubcategoryProperties(IDetailCategoryBuilder& CategoryBuilder)
+void AddSubcategoryProperties(IDetailLayoutBuilder& LayoutBuilder, const FName CategoryName, TMap<const FName, IDetailPropertyRow*>* OutDetailRowsByPropertyName = nullptr)
 {
+	IDetailCategoryBuilder& CategoryBuilder = LayoutBuilder.EditCategory(CategoryName);
+
 	TArray<TSharedRef<IPropertyHandle>> CategoryProperties;
 	CategoryBuilder.GetDefaultProperties(CategoryProperties, true, true);
 
 	//Build map of subcategories to properties
-	const FName SubcategoryName("Subcategory");
+	static const FName SubcategoryName("Subcategory");
 	TMap<FString, TArray<TSharedRef<IPropertyHandle>> > SubcategoryPropertiesMap;
 	for (auto& PropertyHandle : CategoryProperties)
 	{
@@ -40,6 +52,17 @@ void AddSubcategoryProperties(IDetailCategoryBuilder& CategoryBuilder)
 				TArray<TSharedRef<IPropertyHandle>>& PropertyHandles = SubcategoryPropertiesMap.FindOrAdd(Subcategory);
 				PropertyHandles.Add(PropertyHandle);
 			}
+			else
+			{
+				// The property is not in a subcategory, so add it now
+				auto& PropertyRow = CategoryBuilder.AddProperty(PropertyHandle);
+				CustomizePropertyRowVisibility(LayoutBuilder, PropertyHandle, PropertyRow);
+
+				if (OutDetailRowsByPropertyName)
+				{
+					OutDetailRowsByPropertyName->Add(Property->GetFName()) = &PropertyRow;
+				}
+			}
 		}
 	}
 
@@ -51,20 +74,45 @@ void AddSubcategoryProperties(IDetailCategoryBuilder& CategoryBuilder)
 		IDetailGroup& Group = CategoryBuilder.AddGroup(FName(*GroupString), FText::FromString(GroupString));
 		for (auto& PropertyHandle : PropertyHandles)
 		{
-			Group.AddPropertyRow(PropertyHandle);
+			auto& PropertyRow = Group.AddPropertyRow(PropertyHandle);
+			CustomizePropertyRowVisibility(LayoutBuilder, PropertyHandle, PropertyRow);
+
+			if (OutDetailRowsByPropertyName && PropertyHandle->GetProperty())
+			{
+				OutDetailRowsByPropertyName->Add(PropertyHandle->GetProperty()->GetFName()) = &PropertyRow;
+			}
 		}
 	}
 }
 
 void FFoliageTypeDetails::CustomizeDetails(IDetailLayoutBuilder& DetailBuilder)
 {
+	const static FName ReapplyName("Reapply");
 	const static FName PaintingName("Painting");
 	const static FName PlacementName("Placement");
 	const static FName ProceduralName("Procedural");
 	const static FName InstanceSettingsName("InstanceSettings");
-	
-	HideFoliageCategory(DetailBuilder, PaintingName);
-	AddSubcategoryProperties(DetailBuilder.EditCategory(PlacementName));
-	AddSubcategoryProperties(DetailBuilder.EditCategory(ProceduralName));
-	AddSubcategoryProperties(DetailBuilder.EditCategory(InstanceSettingsName));
+
+	FFoliageTypeCustomizationHelpers::HideFoliageCategory(DetailBuilder, ReapplyName);
+	FFoliageTypeCustomizationHelpers::HideFoliageCategory(DetailBuilder, PaintingName);
+
+	TMap<const FName, IDetailPropertyRow*> PropertyRowsByName;
+	AddSubcategoryProperties(DetailBuilder, PlacementName, &PropertyRowsByName);
+	AddSubcategoryProperties(DetailBuilder, ProceduralName);
+	AddSubcategoryProperties(DetailBuilder, InstanceSettingsName);
+
+	// Custom visibility overrides for the scale axes
+	TSharedPtr<IPropertyHandle> Scaling = DetailBuilder.GetProperty(GET_MEMBER_NAME_CHECKED(UFoliageType, Scaling));
+
+	FFoliageTypeCustomizationHelpers::ModifyFoliagePropertyRow(*PropertyRowsByName.Find(GET_MEMBER_NAME_CHECKED(UFoliageType, ScaleX)),
+		TAttribute<EVisibility>::Create(TAttribute<EVisibility>::FGetter::CreateStatic(&FFoliageTypeCustomizationHelpers::GetScaleAxisVisibility, EAxis::X, Scaling)),
+		TAttribute<bool>());
+
+	FFoliageTypeCustomizationHelpers::ModifyFoliagePropertyRow(*PropertyRowsByName.Find(GET_MEMBER_NAME_CHECKED(UFoliageType, ScaleY)),
+		TAttribute<EVisibility>::Create(TAttribute<EVisibility>::FGetter::CreateStatic(&FFoliageTypeCustomizationHelpers::GetScaleAxisVisibility, EAxis::Y, Scaling)),
+		TAttribute<bool>());
+
+	FFoliageTypeCustomizationHelpers::ModifyFoliagePropertyRow(*PropertyRowsByName.Find(GET_MEMBER_NAME_CHECKED(UFoliageType, ScaleZ)),
+		TAttribute<EVisibility>::Create(TAttribute<EVisibility>::FGetter::CreateStatic(&FFoliageTypeCustomizationHelpers::GetScaleAxisVisibility, EAxis::Z, Scaling)),
+		TAttribute<bool>());
 }
\ No newline at end of file
diff --git a/Engine/Source/Editor/FoliageEdit/Private/FoliageTypeFactory.cpp b/Engine/Source/Editor/FoliageEdit/Private/FoliageTypeFactory.cpp
new file mode 100644
--- /dev/null
+++ b/Engine/Source/Editor/FoliageEdit/Private/FoliageTypeFactory.cpp
@@ -0,0 +1,26 @@
+// Copyright 1998-2015 Epic Games, Inc. All Rights Reserved.
+
+#include "UnrealEd.h"
+#include "FoliageTypeFactory.h"
+#include "IAssetTypeActions.h"
+#include "FoliageType_InstancedStaticMesh.h"
+
+UFoliageTypeFactory::UFoliageTypeFactory(const FObjectInitializer& ObjectInitializer)
+	: Super(ObjectInitializer)
+{
+	bCreateNew = true;
+	bEditAfterNew = true;
+	SupportedClass = UFoliageType_InstancedStaticMesh::StaticClass();
+}
+
+UObject* UFoliageTypeFactory::FactoryCreateNew(UClass* Class, UObject* InParent, FName Name, EObjectFlags Flags, UObject* Context, FFeedbackContext* Warn)
+{
+	auto NewFoliageType = NewObject<UFoliageType_InstancedStaticMesh>(InParent, Class, Name, Flags | RF_Transactional);
+
+	return NewFoliageType;
+}
+
+uint32 UFoliageTypeFactory::GetMenuCategories() const
+{
+	return EAssetTypeCategories::Misc;
+}
\ No newline at end of file
diff --git a/Engine/Source/Editor/FoliageEdit/Private/FoliageTypeFactory.h b/Engine/Source/Editor/FoliageEdit/Private/FoliageTypeFactory.h
new file mode 100644
--- /dev/null
+++ b/Engine/Source/Editor/FoliageEdit/Private/FoliageTypeFactory.h
@@ -0,0 +1,21 @@
+// Copyright 1998-2015 Epic Games, Inc. All Rights Reserved.
+
+/**
+* Factory for FoliageType assets
+*/
+
+#pragma once
+
+#include "FoliageTypeFactory.generated.h"
+
+
+UCLASS()
+class UFoliageTypeFactory : public UFactory
+{
+	GENERATED_UCLASS_BODY()
+
+	// UFactory interface
+	virtual UObject* FactoryCreateNew(UClass* Class, UObject* InParent, FName Name, EObjectFlags Flags, UObject* Context, FFeedbackContext* Warn) override;
+	virtual uint32 GetMenuCategories() const;
+	// End of UFactory interface
+};
diff --git a/Engine/Source/Editor/FoliageEdit/Private/FoliageTypePaintingCustomization.cpp b/Engine/Source/Editor/FoliageEdit/Private/FoliageTypePaintingCustomization.cpp
--- a/Engine/Source/Editor/FoliageEdit/Private/FoliageTypePaintingCustomization.cpp
+++ b/Engine/Source/Editor/FoliageEdit/Private/FoliageTypePaintingCustomization.cpp
@@ -5,9 +5,12 @@
 #include "FoliageType.h"
 #include "FoliageEdMode.h"
 #include "FoliageTypePaintingCustomization.h"
+#include "FoliageTypeCustomizationHelpers.h"
 
 #define LOCTEXT_NAMESPACE "FoliageEd_Mode"
 
+DECLARE_DELEGATE_RetVal(EVisibility, FFoliageVisibilityDelegate);
+
 /////////////////////////////////////////////////////
 // FFoliageTypePaintingCustomization 
 TSharedRef<IDetailCustomization> FFoliageTypePaintingCustomization::MakeInstance(FEdModeFoliage* InFoliageEditMode)
@@ -17,134 +20,104 @@ TSharedRef<IDetailCustomization> FFoliageTypePaintingCustomization::MakeInstance
 }
 
 FFoliageTypePaintingCustomization::FFoliageTypePaintingCustomization(FEdModeFoliage* InFoliageEditMode)
+	:FoliageEditMode(InFoliageEditMode)
 {
-	FoliageEditMode = InFoliageEditMode;
 }
 
 void FFoliageTypePaintingCustomization::CustomizeDetails(IDetailLayoutBuilder& DetailLayoutBuilder)
 {
 	// Hide categories we are not going to customize
-	const static FName PaintingName("Painting");
-	const static FName PlacementName("Placement");
-	const static FName ProceduralName("Procedural");
-	HideFoliageCategory(DetailLayoutBuilder, PaintingName);
-	HideFoliageCategory(DetailLayoutBuilder, PlacementName);
-	HideFoliageCategory(DetailLayoutBuilder, ProceduralName);
-	
-	//
-	// Re-add customized properties 
-	//
-		
-	// Get the properties we need
-	TSharedPtr<IPropertyHandle> InvalidProperty;
-	IDetailCategoryBuilder& PaintingCategory = DetailLayoutBuilder.EditCategory("Painting");
-	
-	Density = DetailLayoutBuilder.GetProperty(GET_MEMBER_NAME_CHECKED(UFoliageType, Density));
-	Radius = DetailLayoutBuilder.GetProperty(GET_MEMBER_NAME_CHECKED(UFoliageType, Radius));
-	AlignToNormal = DetailLayoutBuilder.GetProperty(GET_MEMBER_NAME_CHECKED(UFoliageType, AlignToNormal));
-	AlignMaxAngle = DetailLayoutBuilder.GetProperty(GET_MEMBER_NAME_CHECKED(UFoliageType, AlignMaxAngle));
-	RandomYaw = DetailLayoutBuilder.GetProperty(GET_MEMBER_NAME_CHECKED(UFoliageType, RandomYaw));
-	Scaling = DetailLayoutBuilder.GetProperty(GET_MEMBER_NAME_CHECKED(UFoliageType, Scaling));
-	ScaleX = DetailLayoutBuilder.GetProperty(GET_MEMBER_NAME_CHECKED(UFoliageType, ScaleX));
-	ScaleY = DetailLayoutBuilder.GetProperty(GET_MEMBER_NAME_CHECKED(UFoliageType, ScaleY));
-	ScaleZ = DetailLayoutBuilder.GetProperty(GET_MEMBER_NAME_CHECKED(UFoliageType, ScaleZ));
-	ZOffset = DetailLayoutBuilder.GetProperty(GET_MEMBER_NAME_CHECKED(UFoliageType, ZOffset));
-	RandomPitchAngle = DetailLayoutBuilder.GetProperty(GET_MEMBER_NAME_CHECKED(UFoliageType, RandomPitchAngle));
-	GroundSlopeAngle = DetailLayoutBuilder.GetProperty(GET_MEMBER_NAME_CHECKED(UFoliageType, GroundSlopeAngle));
-	Height = DetailLayoutBuilder.GetProperty(GET_MEMBER_NAME_CHECKED(UFoliageType, Height));
-	LandscapeLayers = DetailLayoutBuilder.GetProperty(GET_MEMBER_NAME_CHECKED(UFoliageType, LandscapeLayers));
-	MinimumLayerWeight = DetailLayoutBuilder.GetProperty(GET_MEMBER_NAME_CHECKED(UFoliageType, MinimumLayerWeight));
-	CollisionWithWorld = DetailLayoutBuilder.GetProperty(GET_MEMBER_NAME_CHECKED(UFoliageType, CollisionWithWorld));
-	CollisionScale = DetailLayoutBuilder.GetProperty(GET_MEMBER_NAME_CHECKED(UFoliageType, CollisionScale));
-	VertexColorMask = DetailLayoutBuilder.GetProperty(GET_MEMBER_NAME_CHECKED(UFoliageType, VertexColorMask));
-	VertexColorMaskInvert = DetailLayoutBuilder.GetProperty(GET_MEMBER_NAME_CHECKED(UFoliageType, VertexColorMaskInvert));
-	VertexColorMaskThreshold = DetailLayoutBuilder.GetProperty(GET_MEMBER_NAME_CHECKED(UFoliageType, VertexColorMaskThreshold));
-	
-	//
-	ReapplyDensityAmount = DetailLayoutBuilder.GetProperty(GET_MEMBER_NAME_CHECKED(UFoliageType, ReapplyDensityAmount));
-	ReapplyDensity = DetailLayoutBuilder.GetProperty(GET_MEMBER_NAME_CHECKED(UFoliageType, ReapplyDensity));
-	ReapplyRadius = DetailLayoutBuilder.GetProperty(GET_MEMBER_NAME_CHECKED(UFoliageType, ReapplyRadius));
-	ReapplyAlignToNormal = DetailLayoutBuilder.GetProperty(GET_MEMBER_NAME_CHECKED(UFoliageType, ReapplyAlignToNormal));
-	ReapplyRandomYaw = DetailLayoutBuilder.GetProperty(GET_MEMBER_NAME_CHECKED(UFoliageType, ReapplyRandomYaw));
-	ReapplyScaleX = DetailLayoutBuilder.GetProperty(GET_MEMBER_NAME_CHECKED(UFoliageType, ReapplyScaleX));
-	ReapplyScaleY = DetailLayoutBuilder.GetProperty(GET_MEMBER_NAME_CHECKED(UFoliageType, ReapplyScaleY));
-	ReapplyScaleZ = DetailLayoutBuilder.GetProperty(GET_MEMBER_NAME_CHECKED(UFoliageType, ReapplyScaleZ));
-	ReapplyRandomPitchAngle = DetailLayoutBuilder.GetProperty(GET_MEMBER_NAME_CHECKED(UFoliageType, ReapplyRandomPitchAngle));
-	ReapplyGroundSlopeAngle = DetailLayoutBuilder.GetProperty(GET_MEMBER_NAME_CHECKED(UFoliageType, ReapplyGroundSlope));
-	ReapplyHeight = DetailLayoutBuilder.GetProperty(GET_MEMBER_NAME_CHECKED(UFoliageType, ReapplyHeight));
-	ReapplyLandscapeLayer = DetailLayoutBuilder.GetProperty(GET_MEMBER_NAME_CHECKED(UFoliageType, ReapplyLandscapeLayer));
-	ReapplyZOffset = DetailLayoutBuilder.GetProperty(GET_MEMBER_NAME_CHECKED(UFoliageType, ReapplyZOffset));
-	ReapplyCollisionWithWorld = DetailLayoutBuilder.GetProperty(GET_MEMBER_NAME_CHECKED(UFoliageType, ReapplyCollisionWithWorld));
-	ReapplyVertexColorMask = DetailLayoutBuilder.GetProperty(GET_MEMBER_NAME_CHECKED(UFoliageType, ReapplyVertexColorMask));
+	FFoliageTypeCustomizationHelpers::HideFoliageCategory(DetailLayoutBuilder, "Procedural");
+	FFoliageTypeCustomizationHelpers::HideFoliageCategory(DetailLayoutBuilder, "Reapply");
 	
-	// Density is hidden during Reapply mode
-	AddFoliageProperty(PaintingCategory, Density, InvalidProperty, 
-		TAttribute<EVisibility>::Create(TAttribute<EVisibility>::FGetter::CreateSP(this, &FFoliageTypePaintingCustomization::GetDensityVisibility)), 
+	// Show all the properties with a reapply condition or that depend on another variable to be relevant
+	TMap<const FName, IDetailPropertyRow*> PropertyRowsByName;
+	ShowFoliagePropertiesForCategory(DetailLayoutBuilder, "Painting", PropertyRowsByName);
+	ShowFoliagePropertiesForCategory(DetailLayoutBuilder, "Placement", PropertyRowsByName);
+	ShowFoliagePropertiesForCategory(DetailLayoutBuilder, "InstanceSettings", PropertyRowsByName);
+
+	// Density adjustment factor should only be visible when reapplying
+	FFoliageTypeCustomizationHelpers::ModifyFoliagePropertyRow(*PropertyRowsByName.Find(GET_MEMBER_NAME_CHECKED(UFoliageType, DensityAdjustmentFactor)),
+		TAttribute<EVisibility>::Create(TAttribute<EVisibility>::FGetter::CreateSP(this, &FFoliageTypePaintingCustomization::GetReapplyModeVisibility)),
 		TAttribute<bool>());
-	
-	// Reapply DensityAmount is visible only during Reapply mode
-	AddFoliageProperty(PaintingCategory, ReapplyDensityAmount, ReapplyDensity, 
-		TAttribute<EVisibility>::Create(TAttribute<EVisibility>::FGetter::CreateSP(this, &FFoliageTypePaintingCustomization::GetReapplyDensityAmountVisibility)), 
+
+	// Set the scale visibility attribute for each axis
+	Scaling = DetailLayoutBuilder.GetProperty(GET_MEMBER_NAME_CHECKED(UFoliageType, Scaling));
+	ReapplyScaling = DetailLayoutBuilder.GetProperty(GET_MEMBER_NAME_CHECKED(UFoliageType, ReapplyScaling));
+	FFoliageTypeCustomizationHelpers::ModifyFoliagePropertyRow(*PropertyRowsByName.Find(GET_MEMBER_NAME_CHECKED(UFoliageType, ScaleX)),
+		TAttribute<EVisibility>::Create(TAttribute<EVisibility>::FGetter::CreateSP(this, &FFoliageTypePaintingCustomization::GetScaleVisibility, EAxis::X)),
 		TAttribute<bool>());
 
-	// Radius
-	AddFoliageProperty(PaintingCategory, Radius, ReapplyRadius, TAttribute<EVisibility>(), TAttribute<bool>());
+	FFoliageTypeCustomizationHelpers::ModifyFoliagePropertyRow(*PropertyRowsByName.Find(GET_MEMBER_NAME_CHECKED(UFoliageType, ScaleY)),
+		TAttribute<EVisibility>::Create(TAttribute<EVisibility>::FGetter::CreateSP(this, &FFoliageTypePaintingCustomization::GetScaleVisibility, EAxis::Y)),
+		TAttribute<bool>());
 
-	// Align to normal
-	AddFoliageProperty(PaintingCategory, AlignToNormal, ReapplyAlignToNormal, TAttribute<EVisibility>(), TAttribute<bool>());
-	AddFoliageProperty(PaintingCategory, AlignMaxAngle, InvalidProperty,
-		TAttribute<EVisibility>::Create(TAttribute<EVisibility>::FGetter::CreateSP(this, &FFoliageTypePaintingCustomization::GetAlignMaxAngleVisibility)),
-		TAttribute<bool>::Create(TAttribute<bool>::FGetter::CreateSP(this, &FFoliageTypePaintingCustomization::IsReapplyPropertyEnabled, ReapplyAlignToNormal))); 
-		
-	// Random Yaw
-	AddFoliageProperty(PaintingCategory, RandomYaw, ReapplyRandomYaw, TAttribute<EVisibility>(), TAttribute<bool>());
-		
-	// Scaling
-	AddFoliageProperty(PaintingCategory, Scaling, InvalidProperty, TAttribute<EVisibility>(), TAttribute<bool>());
-	
-	AddFoliageProperty(PaintingCategory, ScaleX, ReapplyScaleX, 
-		TAttribute<EVisibility>::Create(TAttribute<EVisibility>::FGetter::CreateSP(this, &FFoliageTypePaintingCustomization::GetScaleXVisibility)), TAttribute<bool>());
-	
-	AddFoliageProperty(PaintingCategory, ScaleY, ReapplyScaleY, 
-		TAttribute<EVisibility>::Create(TAttribute<EVisibility>::FGetter::CreateSP(this, &FFoliageTypePaintingCustomization::GetScaleYVisibility)), TAttribute<bool>());
-	
-	AddFoliageProperty(PaintingCategory, ScaleZ, ReapplyScaleZ, 
-		TAttribute<EVisibility>::Create(TAttribute<EVisibility>::FGetter::CreateSP(this, &FFoliageTypePaintingCustomization::GetScaleZVisibility)), TAttribute<bool>());
+	FFoliageTypeCustomizationHelpers::ModifyFoliagePropertyRow(*PropertyRowsByName.Find(GET_MEMBER_NAME_CHECKED(UFoliageType, ScaleZ)),
+		TAttribute<EVisibility>::Create(TAttribute<EVisibility>::FGetter::CreateSP(this, &FFoliageTypePaintingCustomization::GetScaleVisibility, EAxis::Z)),
+		TAttribute<bool>());
+}
 
-	// ZOffset
-	AddFoliageProperty(PaintingCategory, ZOffset, ReapplyZOffset, TAttribute<EVisibility>(), TAttribute<bool>());
+void FFoliageTypePaintingCustomization::ShowFoliagePropertiesForCategory(IDetailLayoutBuilder& DetailLayoutBuilder, const FName CategoryName, TMap<const FName, IDetailPropertyRow*>& OutDetailRowsByPropertyName)
+{
+	// Properties that have a ReapplyCondition should be disabled behind the specified property when in reapply mode
+	static const FName ReapplyConditionKey("ReapplyCondition");
 	
-	// Random pitch
-	AddFoliageProperty(PaintingCategory, RandomPitchAngle, ReapplyRandomPitchAngle, TAttribute<EVisibility>(), TAttribute<bool>());
-
-	// Ground slope
-	AddFoliageProperty(PaintingCategory, GroundSlopeAngle, ReapplyGroundSlopeAngle, TAttribute<EVisibility>(), TAttribute<bool>());
+	// Properties with a HideBehind property specified should only be shown if that property is true, non-zero, or not empty
+	static const FName HideBehindKey("HideBehind");
 
-	// Height
-	AddFoliageProperty(PaintingCategory, Height, ReapplyHeight, TAttribute<EVisibility>(), TAttribute<bool>());
-
-	// Landscape layers
-	AddFoliageProperty(PaintingCategory, LandscapeLayers, ReapplyLandscapeLayer, TAttribute<EVisibility>(), TAttribute<bool>());
-	AddFoliageProperty(PaintingCategory, MinimumLayerWeight, InvalidProperty, 
-		TAttribute<EVisibility>(), 
-		TAttribute<bool>::Create(TAttribute<bool>::FGetter::CreateSP(this, &FFoliageTypePaintingCustomization::IsReapplyPropertyEnabled, ReapplyLandscapeLayer))); 
-
-	// Collision with world
-	AddFoliageProperty(PaintingCategory, CollisionWithWorld, ReapplyCollisionWithWorld, TAttribute<EVisibility>(), TAttribute<bool>());
-	AddFoliageProperty(PaintingCategory, CollisionScale, InvalidProperty, 
-		TAttribute<EVisibility>::Create(TAttribute<EVisibility>::FGetter::CreateSP(this, &FFoliageTypePaintingCustomization::GetCollisionScaleVisibility)),
-		TAttribute<bool>::Create(TAttribute<bool>::FGetter::CreateSP(this, &FFoliageTypePaintingCustomization::IsReapplyPropertyEnabled, ReapplyCollisionWithWorld))); 
+	IDetailCategoryBuilder& CategoryBuilder = DetailLayoutBuilder.EditCategory(CategoryName);
+	TArray<TSharedRef<IPropertyHandle>> CategoryProperties;
+	CategoryBuilder.GetDefaultProperties(CategoryProperties, true, true);
 
-	// Vertex color mask
-	AddFoliageProperty(PaintingCategory, VertexColorMask, ReapplyVertexColorMask, TAttribute<EVisibility>(), TAttribute<bool>());
-	
-	AddFoliageProperty(PaintingCategory, VertexColorMaskInvert, InvalidProperty, 
-		TAttribute<EVisibility>::Create(TAttribute<EVisibility>::FGetter::CreateSP(this, &FFoliageTypePaintingCustomization::GetVertexColorMaskDetailsVisibility)),
-		TAttribute<bool>::Create(TAttribute<bool>::FGetter::CreateSP(this, &FFoliageTypePaintingCustomization::IsReapplyPropertyEnabled, ReapplyVertexColorMask))); 
-	
-	AddFoliageProperty(PaintingCategory, VertexColorMaskThreshold, InvalidProperty, 
-		TAttribute<EVisibility>::Create(TAttribute<EVisibility>::FGetter::CreateSP(this, &FFoliageTypePaintingCustomization::GetVertexColorMaskDetailsVisibility)), 
-		TAttribute<bool>::Create(TAttribute<bool>::FGetter::CreateSP(this, &FFoliageTypePaintingCustomization::IsReapplyPropertyEnabled, ReapplyVertexColorMask))); 
+	// Determine whether each property should be shown and how
+	for (auto& PropertyHandle : CategoryProperties)
+	{
+		bool bShowingProperty = false;
+		if (UProperty* Property = PropertyHandle->GetProperty())
+		{
+			// Check to see if this property can be reapplied
+			TSharedPtr<IPropertyHandle> ReapplyConditionPropertyHandle = DetailLayoutBuilder.GetProperty(*Property->GetMetaData(ReapplyConditionKey));
+			if (ReapplyConditionPropertyHandle.IsValid() && ReapplyConditionPropertyHandle->IsValidHandle())
+			{
+				// Create a custom entry that allows explicit enabling/disabling of the property when reapplying
+				TSharedPtr<IPropertyHandle> PropertyHandlePtr = PropertyHandle;
+				OutDetailRowsByPropertyName.FindOrAdd(PropertyHandle->GetProperty()->GetFName()) =
+					&AddFoliageProperty(CategoryBuilder, PropertyHandlePtr, ReapplyConditionPropertyHandle, TAttribute<EVisibility>(), TAttribute<bool>());
+			}
+			else
+			{
+				TSharedPtr<IPropertyHandle> InvalidProperty;
+				TSharedPtr<IPropertyHandle> PropertyHandlePtr = PropertyHandle;
+
+				// Check to see if this property is hidden behind another
+				TSharedPtr<IPropertyHandle> HiddenBehindPropertyHandle = DetailLayoutBuilder.GetProperty(*Property->GetMetaData(HideBehindKey));
+				if (HiddenBehindPropertyHandle.IsValid() && HiddenBehindPropertyHandle->IsValidHandle())
+				{
+					TAttribute<bool> IsEnabledAttribute;
+					ReapplyConditionPropertyHandle = DetailLayoutBuilder.GetProperty(*HiddenBehindPropertyHandle->GetProperty()->GetMetaData(ReapplyConditionKey));
+					if (ReapplyConditionPropertyHandle.IsValid() && ReapplyConditionPropertyHandle->IsValidHandle())
+					{
+						// If the property this is hidden behind has a reapply condition, disable this when the condition is false
+						IsEnabledAttribute = TAttribute<bool>::Create(TAttribute<bool>::FGetter::CreateSP(this, &FFoliageTypePaintingCustomization::IsReapplyPropertyEnabled, ReapplyConditionPropertyHandle));
+					}
+
+					TAttribute<EVisibility> VisibilityAttribute;
+					GetHiddenPropertyVisibility(HiddenBehindPropertyHandle, !IsEnabledAttribute.IsSet(), VisibilityAttribute);
+
+					OutDetailRowsByPropertyName.FindOrAdd(PropertyHandle->GetProperty()->GetFName()) =
+						&AddFoliageProperty(CategoryBuilder, PropertyHandlePtr, InvalidProperty, VisibilityAttribute, IsEnabledAttribute);
+				}
+				else
+				{
+					// This property cannot be reapplied and isn't hidden behind anything, so show it whenever the reapply tool isn't active
+					OutDetailRowsByPropertyName.FindOrAdd(PropertyHandle->GetProperty()->GetFName()) =
+						&AddFoliageProperty(CategoryBuilder, PropertyHandlePtr, InvalidProperty,
+						TAttribute<EVisibility>::Create(TAttribute<EVisibility>::FGetter::CreateSP(this, &FFoliageTypePaintingCustomization::GetNonReapplyPropertyVisibility)),
+						TAttribute<bool>());
+				}
+			}
+		}
+	}
 }
 
 IDetailPropertyRow& FFoliageTypePaintingCustomization::AddFoliageProperty(
@@ -167,7 +140,7 @@ IDetailPropertyRow& FFoliageTypePaintingCustomization::AddFoliageProperty(
 		ValueWidget->SetEnabled(IsEnabled);
 		NameWidget->SetEnabled(IsEnabled);
 
-		PropertyRow.CustomWidget()
+		PropertyRow.CustomWidget(true)
 		.NameContent()
 		.MinDesiredWidth(Row.NameWidget.MinWidth)
 		.MaxDesiredWidth(Row.NameWidget.MaxWidth)
@@ -179,7 +152,6 @@ IDetailPropertyRow& FFoliageTypePaintingCustomization::AddFoliageProperty(
 				SNew(SCheckBox)
 				.IsChecked(this, &FFoliageTypePaintingCustomization::GetReapplyPropertyState, ReapplyProperty)
 				.OnCheckStateChanged(this, &FFoliageTypePaintingCustomization::OnReapplyPropertyStateChanged, ReapplyProperty)
-				.IsChecked(this, &FFoliageTypePaintingCustomization::GetReapplyPropertyState, ReapplyProperty)
 				.Visibility(this, &FFoliageTypePaintingCustomization::GetReapplyModeVisibility)
 				.ToolTipText(ReapplyProperty->GetToolTipText())
 			]
@@ -209,75 +181,42 @@ IDetailPropertyRow& FFoliageTypePaintingCustomization::AddFoliageProperty(
 		PropertyRow.Visibility(InVisibility);
 	}
 	
-	//
 	return PropertyRow;
 }
 
-void FFoliageTypePaintingCustomization::HideFoliageCategory(IDetailLayoutBuilder& DetailLayoutBuilder, FName CategoryName)
+void FFoliageTypePaintingCustomization::GetHiddenPropertyVisibility(const TSharedPtr<IPropertyHandle>& PropertyHandle, bool bHideInReapplyTool, TAttribute<EVisibility>& OutVisibility) const
 {
-	TArray<TSharedRef<IPropertyHandle>> CategoryProperties;
-	DetailLayoutBuilder.EditCategory(CategoryName).GetDefaultProperties(CategoryProperties, true, true);
-
-	for (auto& PropertyHandle : CategoryProperties)
-	{
-		DetailLayoutBuilder.HideProperty(PropertyHandle);
-	}
-}
+	TAttribute<EVisibility>::FGetter VisibilityGetter;
+	FFoliageTypeCustomizationHelpers::BindHiddenPropertyVisibilityGetter(PropertyHandle, VisibilityGetter);
 
-EVisibility FFoliageTypePaintingCustomization::GetAlignMaxAngleVisibility() const
-{
-	bool bAlightToNormal;
-	if (AlignToNormal->GetValue(bAlightToNormal) == FPropertyAccess::Success)
+	if (bHideInReapplyTool)
 	{
-		return bAlightToNormal ? EVisibility::Visible : EVisibility::Collapsed;
-	}
-
-	return EVisibility::Visible;
-}
-
-EVisibility FFoliageTypePaintingCustomization::GetScaleXVisibility() const
-{
-	return EVisibility::Visible;
-}
-
-EVisibility FFoliageTypePaintingCustomization::GetScaleYVisibility() const
-{
-	uint8 ScalingValue;
-	if (Scaling->GetValue(ScalingValue) == FPropertyAccess::Success)
-	{
-		return (ScalingValue == (uint8)EFoliageScaling::Uniform) ? EVisibility::Collapsed : EVisibility::Visible;
+		// In addition to hiding it it behind the given property, only show this in the reapply tool
+		OutVisibility = TAttribute<EVisibility>::Create(TAttribute<EVisibility>::FGetter::CreateLambda([=]
+		{
+			if (!FoliageEditMode->UISettings.GetReapplyToolSelected() && VisibilityGetter.IsBound())
+			{
+				const EVisibility ReturnVal = VisibilityGetter.Execute();
+				return ReturnVal;
+			}
+			return EVisibility::Collapsed;
+		}));
 	}
-	return EVisibility::Visible;
-}
-
-EVisibility FFoliageTypePaintingCustomization::GetScaleZVisibility() const
-{
-	uint8 ScalingValue;
-	if (Scaling->GetValue(ScalingValue) == FPropertyAccess::Success)
+	else
 	{
-		return (ScalingValue == (uint8)EFoliageScaling::Uniform) ? EVisibility::Collapsed : EVisibility::Visible;
+		OutVisibility.Bind(VisibilityGetter);
 	}
-	return EVisibility::Visible;
 }
 
-EVisibility FFoliageTypePaintingCustomization::GetCollisionScaleVisibility() const
+EVisibility FFoliageTypePaintingCustomization::GetScaleVisibility(EAxis::Type Axis) const
 {
-	bool bCollisionWithWorld = false;
-	if (CollisionWithWorld->GetValue(bCollisionWithWorld) == FPropertyAccess::Success)
+	// In reapply mode we only want to show these if scaling is being reapplied
+	if (IsReapplyPropertyEnabled(ReapplyScaling) || !FoliageEditMode->UISettings.GetReapplyToolSelected())
 	{
-		return bCollisionWithWorld ? EVisibility::Visible : EVisibility::Collapsed;
+		return FFoliageTypeCustomizationHelpers::GetScaleAxisVisibility(Axis, Scaling);
 	}
-	return EVisibility::Visible;
-}
 
-EVisibility FFoliageTypePaintingCustomization::GetVertexColorMaskDetailsVisibility() const
-{
-	uint8 ColorMask;
-	if (VertexColorMask->GetValue(ColorMask) == FPropertyAccess::Success)
-	{
-		return ColorMask != FOLIAGEVERTEXCOLORMASK_Disabled ? EVisibility::Visible : EVisibility::Collapsed;
-	}
-	return EVisibility::Visible;
+	return EVisibility::Collapsed;
 }
 
 EVisibility FFoliageTypePaintingCustomization::GetReapplyModeVisibility() const
@@ -318,17 +257,11 @@ bool FFoliageTypePaintingCustomization::IsReapplyPropertyEnabled(TSharedPtr<IPro
 	return true;
 }
 
-EVisibility FFoliageTypePaintingCustomization::GetDensityVisibility() const
+EVisibility FFoliageTypePaintingCustomization::GetNonReapplyPropertyVisibility() const
 {
-	// mutually exclusive with ReapplyDensityAmmount
+	// Visible if the reapply tool is not active
 	return !FoliageEditMode->UISettings.GetReapplyToolSelected() ? EVisibility::Visible : EVisibility::Collapsed;
 }
 
-EVisibility FFoliageTypePaintingCustomization::GetReapplyDensityAmountVisibility() const
-{
-	// mutually exclusive with Density
-	return FoliageEditMode->UISettings.GetReapplyToolSelected() ? EVisibility::Visible : EVisibility::Collapsed;
-}
-
 
 #undef LOCTEXT_NAMESPACE
\ No newline at end of file
diff --git a/Engine/Source/Editor/FoliageEdit/Private/FoliageTypePaintingCustomization.h b/Engine/Source/Editor/FoliageEdit/Private/FoliageTypePaintingCustomization.h
--- a/Engine/Source/Editor/FoliageEdit/Private/FoliageTypePaintingCustomization.h
+++ b/Engine/Source/Editor/FoliageEdit/Private/FoliageTypePaintingCustomization.h
@@ -19,21 +19,15 @@ public:
 private:
 	FFoliageTypePaintingCustomization(FEdModeFoliage* InFoliageEditMode);
 	
-	EVisibility GetAlignMaxAngleVisibility() const;
-	EVisibility GetScaleXVisibility() const;
-	EVisibility GetScaleYVisibility() const;
-	EVisibility GetScaleZVisibility() const;
-	EVisibility GetCollisionScaleVisibility() const;
-	EVisibility GetVertexColorMaskDetailsVisibility() const;
-	
+	EVisibility GetScaleVisibility(EAxis::Type Axis) const;
+	EVisibility GetNonReapplyPropertyVisibility() const;
 	EVisibility GetReapplyModeVisibility() const;
+
 	ECheckBoxState GetReapplyPropertyState(TSharedPtr<IPropertyHandle> ReapplyProperty) const;
 	void OnReapplyPropertyStateChanged(ECheckBoxState CheckState, TSharedPtr<IPropertyHandle> ReapplyProperty);
 	bool IsReapplyPropertyEnabled(TSharedPtr<IPropertyHandle> ReapplyProperty) const;
 
-	EVisibility GetDensityVisibility() const;
-	EVisibility GetReapplyDensityAmountVisibility() const;
-
+	/** Adds a property that can optionally be hidden behind another when the reapply tool is active */
 	IDetailPropertyRow& AddFoliageProperty(
 		IDetailCategoryBuilder& Category,
 		TSharedPtr<IPropertyHandle>& Property, 
@@ -41,47 +35,15 @@ private:
 		const TAttribute<EVisibility>& InVisibility,
 		const TAttribute<bool>& InEnabled);
 
-	static void HideFoliageCategory(IDetailLayoutBuilder& DetailLayout,	FName CategoryName);
-	
+	void ShowFoliagePropertiesForCategory(IDetailLayoutBuilder& DetailLayoutBuilder, const FName CategoryName, TMap<const FName, IDetailPropertyRow*>& OutDetailRowsByPropertyName);
+
+	/** Gets the visibility attribute for property that is hidden behind another */
+	void GetHiddenPropertyVisibility(const TSharedPtr<IPropertyHandle>& PropertyHandle, bool bHideInReapplyTool, TAttribute<EVisibility>& OutVisibility) const;
+
 private:
 	/** Pointer to the foliage edit mode. */
 	FEdModeFoliage* FoliageEditMode;
 
-	TSharedPtr<IPropertyHandle> Density;
-	TSharedPtr<IPropertyHandle> Radius;
-	TSharedPtr<IPropertyHandle> AlignToNormal;
-	TSharedPtr<IPropertyHandle> AlignMaxAngle;
-	TSharedPtr<IPropertyHandle> RandomYaw;
 	TSharedPtr<IPropertyHandle> Scaling;
-	TSharedPtr<IPropertyHandle> ScaleX;
-	TSharedPtr<IPropertyHandle> ScaleY;
-	TSharedPtr<IPropertyHandle> ScaleZ;
-	TSharedPtr<IPropertyHandle> ZOffset;
-	TSharedPtr<IPropertyHandle> RandomPitchAngle;
-	TSharedPtr<IPropertyHandle> GroundSlopeAngle;
-	TSharedPtr<IPropertyHandle> Height;
-	TSharedPtr<IPropertyHandle>	LandscapeLayers;
-	TSharedPtr<IPropertyHandle>	MinimumLayerWeight;
-	TSharedPtr<IPropertyHandle> CollisionWithWorld;
-	TSharedPtr<IPropertyHandle> CollisionScale;
-	TSharedPtr<IPropertyHandle> VertexColorMask;
-	TSharedPtr<IPropertyHandle> VertexColorMaskInvert;
-	TSharedPtr<IPropertyHandle> VertexColorMaskThreshold;
-
-	//
-	TSharedPtr<IPropertyHandle> ReapplyDensity;
-	TSharedPtr<IPropertyHandle> ReapplyDensityAmount;
-	TSharedPtr<IPropertyHandle> ReapplyRadius;
-	TSharedPtr<IPropertyHandle> ReapplyAlignToNormal;
-	TSharedPtr<IPropertyHandle> ReapplyRandomYaw;
-	TSharedPtr<IPropertyHandle> ReapplyScaleX;
-	TSharedPtr<IPropertyHandle> ReapplyScaleY;
-	TSharedPtr<IPropertyHandle> ReapplyScaleZ;
-	TSharedPtr<IPropertyHandle> ReapplyRandomPitchAngle;
-	TSharedPtr<IPropertyHandle> ReapplyGroundSlopeAngle;
-	TSharedPtr<IPropertyHandle> ReapplyHeight;
-	TSharedPtr<IPropertyHandle> ReapplyLandscapeLayer;
-	TSharedPtr<IPropertyHandle> ReapplyZOffset;
-	TSharedPtr<IPropertyHandle> ReapplyCollisionWithWorld;
-	TSharedPtr<IPropertyHandle> ReapplyVertexColorMask;
+	TSharedPtr<IPropertyHandle> ReapplyScaling;
 };
diff --git a/Engine/Source/Editor/FoliageEdit/Private/ProceduralFoliageSpawnerFactory.h b/Engine/Source/Editor/FoliageEdit/Private/ProceduralFoliageSpawnerFactory.h
--- a/Engine/Source/Editor/FoliageEdit/Private/ProceduralFoliageSpawnerFactory.h
+++ b/Engine/Source/Editor/FoliageEdit/Private/ProceduralFoliageSpawnerFactory.h
@@ -1,7 +1,7 @@
 // Copyright 1998-2015 Epic Games, Inc. All Rights Reserved.
 
 /**
- * Factory for ProceduralFoliage assets
+ * Factory for ProceduralFoliageSpawner assets
  */
 
 #pragma once
diff --git a/Engine/Source/Editor/FoliageEdit/Private/SFoliageEdit.cpp b/Engine/Source/Editor/FoliageEdit/Private/SFoliageEdit.cpp
--- a/Engine/Source/Editor/FoliageEdit/Private/SFoliageEdit.cpp
+++ b/Engine/Source/Editor/FoliageEdit/Private/SFoliageEdit.cpp
@@ -19,7 +19,12 @@
 #include "Engine/StaticMesh.h"
 #include "SFoliageEditMeshRow.h"
 #include "FoliageTypePaintingCustomization.h"
+#include "DetailLayoutBuilder.h"
+#include "SNumericEntryBox.h"
 
+#include "FoliageType_InstancedStaticMesh.h"
+#include "Editor/PropertyEditor/Public/PropertyCustomizationHelpers.h"
+#include "LevelEditor.h"
 
 #define LOCTEXT_NAMESPACE "FoliageEd_Mode"
 
@@ -89,6 +94,7 @@ BEGIN_SLATE_FUNCTION_BUILD_OPTIMIZATION
 void SFoliageEdit::Construct(const FArguments& InArgs)
 {
 	FoliageEditMode = (FEdModeFoliage*)GLevelEditorModeTools().GetActiveMode(FBuiltinEditorModes::EM_Foliage);
+	FoliageEditMode->OnToolChanged.AddSP(this, &SFoliageEdit::HandleOnToolChanged);
 
 	FFoliageEditCommands::Register();
 
@@ -100,12 +106,21 @@ void SFoliageEdit::Construct(const FArguments& InArgs)
 	FDetailsViewArgs Args(false, false, false, FDetailsViewArgs::HideNameArea, true);
 	Args.bShowActorLabel = false;
 	MeshDetailsWidget = PropertyModule.CreateDetailView(Args);
+
+	// We want to use our own customization for UFoliageType
+	MeshDetailsWidget->UnregisterInstancedCustomPropertyLayout(UFoliageType::StaticClass());
 	MeshDetailsWidget->RegisterInstancedCustomPropertyLayout(UFoliageType::StaticClass(), 
 		FOnGetDetailCustomizationInstance::CreateStatic(&FFoliageTypePaintingCustomization::MakeInstance, FoliageEditMode)
 		);
 
 	// Everything (or almost) uses this padding, change it to expand the padding.
-	FMargin StandardPadding(0.0f, 4.0f, 0.0f, 4.0f);
+	FMargin StandardPadding(6.f, 3.f);
+	FMargin StandardLeftPadding(6.f, 3.f, 3.f, 3.f);
+	FMargin StandardRightPadding(3.f, 3.f, 6.f, 3.f);
+
+	FSlateFontInfo StandardFont = FEditorStyle::GetFontStyle(TEXT("PropertyWindow.NormalFont"));
+
+	const FText BlankText = LOCTEXT("Blank", "");
 
 	this->ChildSlot
 	[
@@ -122,26 +137,32 @@ void SFoliageEdit::Construct(const FArguments& InArgs)
 				BuildToolBar()
 			]
 
+			// Brush Size
 			+ SVerticalBox::Slot()
 			.AutoHeight()
 			[
 				SNew(SHorizontalBox)
 				.ToolTipText(LOCTEXT("BrushSize_Tooltip", "The size of the foliage brush"))
+				.Visibility(this, &SFoliageEdit::GetVisibility_Radius)
 
 				+ SHorizontalBox::Slot()
+				.Padding(StandardLeftPadding)
 				.FillWidth(1.0f)
-				.Padding(StandardPadding)
+				.VAlign(VAlign_Center)
 				[
 					SNew(STextBlock)
-					.Visibility(this, &SFoliageEdit::GetVisibility_Radius)
 					.Text(LOCTEXT("BrushSize", "Brush Size"))
+					.Font(StandardFont)
 				]
 				+ SHorizontalBox::Slot()
+				.Padding(StandardRightPadding)
 				.FillWidth(2.0f)
-				.Padding(StandardPadding)
+				.MaxWidth(100.f)
+				.VAlign(VAlign_Center)
 				[
-					SNew(SSpinBox<float>)
-					.Visibility(this, &SFoliageEdit::GetVisibility_Radius)
+					SNew(SNumericEntryBox<float>)
+					.Font(StandardFont)
+					.AllowSpin(true)
 					.MinValue(0.0f)
 					.MaxValue(65536.0f)
 					.MaxSliderValue(8192.0f)
@@ -151,26 +172,32 @@ void SFoliageEdit::Construct(const FArguments& InArgs)
 				]
 			]
 
+			// Paint Density
 			+ SVerticalBox::Slot()
 			.AutoHeight()
 			[
 				SNew(SHorizontalBox)
 				.ToolTipText(LOCTEXT("PaintDensity_Tooltip", "The density of foliage to paint. This is a multiplier for the individual foliage type's density specifier."))
+				.Visibility(this, &SFoliageEdit::GetVisibility_PaintDensity)
 
 				+ SHorizontalBox::Slot()
+				.Padding(StandardLeftPadding)
 				.FillWidth(1.0f)
-				.Padding(StandardPadding)
+				.VAlign(VAlign_Center)
 				[
 					SNew(STextBlock)
-					.Visibility(this, &SFoliageEdit::GetVisibility_PaintDensity)
 					.Text(LOCTEXT("PaintDensity", "Paint Density"))
+					.Font(StandardFont)
 				]
 				+ SHorizontalBox::Slot()
+				.Padding(StandardRightPadding)
 				.FillWidth(2.0f)
-				.Padding(StandardPadding)
+				.MaxWidth(100.f)
+				.VAlign(VAlign_Center)
 				[
-					SNew(SSpinBox<float>)
-					.Visibility(this, &SFoliageEdit::GetVisibility_PaintDensity)
+					SNew(SNumericEntryBox<float>)
+					.Font(StandardFont)
+					.AllowSpin(true)
 					.MinValue(0.0f)
 					.MaxValue(2.0f)
 					.MaxSliderValue(1.0f)
@@ -179,26 +206,32 @@ void SFoliageEdit::Construct(const FArguments& InArgs)
 				]
 			]
 
+			// Erase Density
 			+ SVerticalBox::Slot()
 			.AutoHeight()
 			[
 				SNew(SHorizontalBox)
 				.ToolTipText(LOCTEXT("EraseDensity_Tooltip", "The density of foliage to leave behind when erasing with the Shift key held. 0 will remove all foliage."))
+				.Visibility(this, &SFoliageEdit::GetVisibility_EraseDensity)
 
 				+ SHorizontalBox::Slot()
+				.Padding(StandardLeftPadding)
 				.FillWidth(1.0f)
-				.Padding(StandardPadding)
+				.VAlign(VAlign_Center)
 				[
 					SNew(STextBlock)
-					.Visibility(this, &SFoliageEdit::GetVisibility_EraseDensity)
 					.Text(LOCTEXT("EraseDensity", "Erase Density"))
+					.Font(StandardFont)
 				]
 				+ SHorizontalBox::Slot()
+				.Padding(StandardRightPadding)
 				.FillWidth(2.0f)
-				.Padding(StandardPadding)
+				.MaxWidth(100.f)
+				.VAlign(VAlign_Center)
 				[
-					SNew(SSpinBox<float>)
-					.Visibility(this, &SFoliageEdit::GetVisibility_EraseDensity)
+					SNew(SNumericEntryBox<float>)
+					.Font(StandardFont)
+					.AllowSpin(true)
 					.MinValue(0.0f)
 					.MaxValue(2.0f)
 					.MaxSliderValue(1.0f)
@@ -207,97 +240,145 @@ void SFoliageEdit::Construct(const FArguments& InArgs)
 				]
 			]
 
+			// Filter
 			+ SVerticalBox::Slot()
 			.AutoHeight()
-			.Padding(StandardPadding)
 			[
-				SNew(STextBlock)
+				SNew(SHorizontalBox)
 				.Visibility(this, &SFoliageEdit::GetVisibility_Filters)
-				.Text(LOCTEXT("Filter", "Filter"))
-			]
-
-			+ SVerticalBox::Slot()
-			.AutoHeight()
-			.Padding(StandardPadding)
-			[
-				SNew(SWrapBox)
-				.UseAllottedWidth(true)
 
-				+ SWrapBox::Slot()
-				.Padding(0, 0, 16, 0)
+				+ SHorizontalBox::Slot()
+				.Padding(StandardLeftPadding)
+				.VAlign(VAlign_Center)
+				.FillWidth(1.0f)
 				[
-					SNew(SCheckBox)
-					.Visibility(this, &SFoliageEdit::GetVisibility_Filters)
-					.OnCheckStateChanged(this, &SFoliageEdit::OnCheckStateChanged_Landscape)
-					.IsChecked(this, &SFoliageEdit::GetCheckState_Landscape)
-					.ToolTipText(LOCTEXT("FilterLandscape_Tooltip", "Place foliage on Landscape"))
-					[
-						SNew(STextBlock)
-						.Text(LOCTEXT("Landscape", "Landscape"))
-					]
+					SNew(STextBlock)
+					.Text(this, &SFoliageEdit::GetFilterText)
+					.ToolTipText(this, &SFoliageEdit::GetFilterTooltipText)
+					.Font(StandardFont)
 				]
 
-				+ SWrapBox::Slot()
-				.Padding(0, 0, 16, 0)
+				+ SHorizontalBox::Slot()
+				.VAlign(VAlign_Center)
+				.FillWidth(2.0f)
 				[
-					SNew(SCheckBox)
-					.Visibility(this, &SFoliageEdit::GetVisibility_Filters)
-					.OnCheckStateChanged(this, &SFoliageEdit::OnCheckStateChanged_StaticMesh)
-					.IsChecked(this, &SFoliageEdit::GetCheckState_StaticMesh)
-					.ToolTipText(LOCTEXT("FilterStaticMesh_Tooltip", "Place foliage on StaticMeshes"))
-					[
-						SNew(STextBlock)
-						.Text(LOCTEXT("StaticMeshes", "Static Meshes"))
-					]
-				]
+					SNew(SWrapBox)
+					.UseAllottedWidth(true)
 
-				+ SWrapBox::Slot()
-				.Padding(0, 0, 16, 0)
-				[
-					SNew(SCheckBox)
-					.Visibility(this, &SFoliageEdit::GetVisibility_Filters)
-					.OnCheckStateChanged(this, &SFoliageEdit::OnCheckStateChanged_BSP)
-					.IsChecked(this, &SFoliageEdit::GetCheckState_BSP)
-					.ToolTipText(LOCTEXT("FilterBSP_Tooltip", "Place foliage on BSP"))
+					+ SWrapBox::Slot()
+					.Padding(FMargin(0.f, 0.f, 6.f, 0.f))
 					[
-						SNew(STextBlock)
-						.Text(LOCTEXT("BSP", "BSP"))
+						SNew(SVerticalBox)
+						
+						+ SVerticalBox::Slot()
+						.Padding(FMargin(0.f, 5.f, 0.f, 0.f))
+						[
+							SNew(SCheckBox)
+							.Visibility(this, &SFoliageEdit::GetVisibility_Filters)
+							.OnCheckStateChanged(this, &SFoliageEdit::OnCheckStateChanged_Landscape)
+							.IsChecked(this, &SFoliageEdit::GetCheckState_Landscape)
+							.ToolTipText(this, &SFoliageEdit::GetTooltipText_Landscape)
+							[
+								SNew(STextBlock)
+								.Text(LOCTEXT("Landscape", "Landscape"))
+								.Font(StandardFont)
+							]
+						]
+						
+						+ SVerticalBox::Slot()
+						.Padding(FMargin(0.f, 5.f, 0.f, 0.f))
+						[
+							SNew(SCheckBox)
+							.Visibility(this, &SFoliageEdit::GetVisibility_Filters)
+							.OnCheckStateChanged(this, &SFoliageEdit::OnCheckStateChanged_BSP)
+							.IsChecked(this, &SFoliageEdit::GetCheckState_BSP)
+							.ToolTipText(this, &SFoliageEdit::GetTooltipText_BSP)
+							[
+								SNew(STextBlock)
+								.Text(LOCTEXT("BSP", "BSP"))
+								.Font(StandardFont)
+							]
+						]
 					]
-				]
-
-				+ SWrapBox::Slot()
-				.Padding(0, 0, 16, 0)
-				[
-					SNew(SCheckBox)
-					.Visibility(this, &SFoliageEdit::GetVisibility_Filters)
-					.OnCheckStateChanged(this, &SFoliageEdit::OnCheckStateChanged_Translucent)
-					.IsChecked(this, &SFoliageEdit::GetCheckState_Translucent)
-					.ToolTipText(LOCTEXT("FilterTranslucent_Tooltip", "Place foliage on translucent surfaces"))
+					
+					+ SWrapBox::Slot()
+					.Padding(FMargin(0.f, 0.f, 6.f, 0.f))
 					[
-						SNew(STextBlock)
-						.Text(LOCTEXT("Translucent", "Translucent"))
+						SNew(SVerticalBox)
+
+						+ SVerticalBox::Slot()
+						.Padding(FMargin(0.f, 5.f, 0.f, 0.f))
+						[
+							SNew(SCheckBox)
+							.Visibility(this, &SFoliageEdit::GetVisibility_Filters)
+							.OnCheckStateChanged(this, &SFoliageEdit::OnCheckStateChanged_StaticMesh)
+							.IsChecked(this, &SFoliageEdit::GetCheckState_StaticMesh)
+							.ToolTipText(this, &SFoliageEdit::GetTooltipText_StaticMesh)
+							[
+								SNew(STextBlock)
+								.Text(LOCTEXT("StaticMeshes", "Static Meshes"))
+								.Font(StandardFont)
+							]
+						]
+
+						+ SVerticalBox::Slot()
+						.Padding(FMargin(0.f, 5.f, 0.f, 0.f))
+						[
+							SNew(SCheckBox)
+							.Visibility(this, &SFoliageEdit::GetVisibility_Filters)
+							.OnCheckStateChanged(this, &SFoliageEdit::OnCheckStateChanged_Translucent)
+							.IsChecked(this, &SFoliageEdit::GetCheckState_Translucent)
+							.ToolTipText(this, &SFoliageEdit::GetTooltipText_Translucent)
+							[
+								SNew(STextBlock)
+								.Text(LOCTEXT("Translucent", "Translucent"))
+								.Font(StandardFont)
+							]
+						]
 					]
 				]
 			]
 
 			+ SVerticalBox::Slot()
 			.AutoHeight()
+			.Padding(StandardPadding)
+			[
+				SNew(SSeparator)
+				.Orientation(Orient_Horizontal)
+			]
+
+			+ SVerticalBox::Slot()
+			.Padding(StandardPadding)
+			.AutoHeight()
+			.HAlign(HAlign_Left)
 			[
-				SNew(SHeaderRow)
-				+ SHeaderRow::Column(TEXT("ToggleMeshes"))
+				SAssignNew(AddFoliageTypeCombo, SComboButton)
+				.ForegroundColor(FLinearColor::White)
+				.ButtonStyle(FEditorStyle::Get(), "FlatButton.Success")
+				.OnGetMenuContent(this, &SFoliageEdit::OnGetAddFoliageMeshAssetPicker)
+				.ContentPadding(FMargin(0.f))
+				.ButtonContent()
 				[
-					SNew(SCheckBox)
-					.IsChecked(this, &SFoliageEdit::GetState_AllMeshes)
-					.OnCheckStateChanged(this, &SFoliageEdit::OnCheckStateChanged_AllMeshes)
+					SNew(SHorizontalBox)
+					+ SHorizontalBox::Slot()
+					.VAlign(VAlign_Center)
+					.AutoWidth()
+					.Padding(1.f)
+					[
+						SNew(STextBlock)
+						.TextStyle(FEditorStyle::Get(), "FoliageEditMode.AddFoliageType.Text")
+						.Font(FEditorStyle::Get().GetFontStyle("FontAwesome.9"))
+						.Text(FText::FromString(FString(TEXT("\xf067"))) /*fa-plus*/)
+					]
+					+ SHorizontalBox::Slot()
+					.VAlign(VAlign_Center)
+					.Padding(1.f)
+					[
+						SNew(STextBlock)
+						.Text(LOCTEXT("AddFoliageTypeButtonLabel", "Add Foliage Type"))
+						.TextStyle(FEditorStyle::Get(), "FoliageEditMode.AddFoliageType.Text")
+					]
 				]
-				.HeaderContentPadding(FMargin(0,1,0,1))
-				.FixedWidth(24)
-								
-				+ SHeaderRow::Column(TEXT("Meshes"))
-				.HeaderContentPadding(FMargin(10,1,0,1))
-				.DefaultLabel(this, &SFoliageEdit::GetMeshesHeaderText)
-				.SortMode(this, &SFoliageEdit::GetMeshColumnSortMode)
-				.OnSort(this, &SFoliageEdit::OnMeshesColumnSortModeChanged)
 			]
 
 			+ SVerticalBox::Slot()
@@ -324,6 +405,7 @@ void SFoliageEdit::Construct(const FArguments& InArgs)
 			[
 				SNew(SSplitter)
 				.Orientation(Orient_Vertical)
+				.Style(FEditorStyle::Get(), "DetailsView.Splitter")
 
 				+SSplitter::Slot()
 				.Value(0.3f)
@@ -336,17 +418,38 @@ void SFoliageEdit::Construct(const FArguments& InArgs)
 					.OnContextMenuOpening(this, &SFoliageEdit::ConstructFoliageMeshContextMenu)
 					.OnSelectionChanged(this, &SFoliageEdit::MeshTreeOnSelectionChanged)
 					.HeaderRow
-						(
-							SNew( SHeaderRow ).Visibility(EVisibility::Collapsed)
-							// toggle mesh for painting checkbox
-							+SHeaderRow::Column(FoliageMeshColumns::ColumnID_ToggleMesh).FixedWidth(24.0f)
-							// mesh text label
-							+SHeaderRow::Column(FoliageMeshColumns::ColumnID_MeshLabel).FillWidth(1.0f)
-							// instance count
-							+SHeaderRow::Column(FoliageMeshColumns::ColumnID_InstanceCount).FillWidth(0.3f)
-							// save Foliage asset
-							+SHeaderRow::Column(FoliageMeshColumns::ColumnID_Save).FixedWidth(24.0f)
-						)
+					(
+						// toggle mesh for painting checkbox
+						SAssignNew(MeshTreeHeaderRowWidget, SHeaderRow)
+						+ SHeaderRow::Column(FoliageMeshColumns::ColumnID_ToggleMesh)
+						[
+							SNew(SCheckBox)
+							.IsChecked(this, &SFoliageEdit::GetState_AllMeshes)
+							.OnCheckStateChanged(this, &SFoliageEdit::OnCheckStateChanged_AllMeshes)
+						]
+						.DefaultLabel(BlankText)
+						.HeaderContentPadding(FMargin(0, 1, 0, 1))
+						.HAlignHeader(HAlign_Center)
+						.HAlignCell(HAlign_Center)
+						.FixedWidth(24)
+
+						// mesh text label
+						+ SHeaderRow::Column(FoliageMeshColumns::ColumnID_MeshLabel)
+						.HeaderContentPadding(FMargin(10, 1, 0, 1))
+						.SortMode(this, &SFoliageEdit::GetMeshColumnSortMode)
+						.OnSort(this, &SFoliageEdit::OnMeshesColumnSortModeChanged)
+						.DefaultLabel(this, &SFoliageEdit::GetMeshesHeaderText)
+
+						// instance count
+						+ SHeaderRow::Column(FoliageMeshColumns::ColumnID_InstanceCount)
+						.HeaderContentPadding(FMargin(10, 1, 0, 1))
+						.DefaultLabel(LOCTEXT("InstanceCount", "Count"))
+						.FixedWidth(60.f)
+
+						// save Foliage asset
+						+SHeaderRow::Column(FoliageMeshColumns::ColumnID_Save).FixedWidth(24.0f)
+						.DefaultLabel(BlankText)
+					)
 				]
 					
 				+SSplitter::Slot()
@@ -380,10 +483,6 @@ void SFoliageEdit::Construct(const FArguments& InArgs)
 }
 END_SLATE_FUNCTION_BUILD_OPTIMIZATION
 
-SFoliageEdit::~SFoliageEdit()
-{
-}
-
 void SFoliageEdit::RefreshFullList()
 {
 	MeshTreeWidget->RequestTreeRefresh();
@@ -501,6 +600,24 @@ TSharedRef<SWidget> SFoliageEdit::BuildToolBar()
 		];
 }
 
+void SFoliageEdit::AddFoliageType(const FAssetData& AssetData)
+{
+	if (AddFoliageTypeCombo.IsValid())
+	{
+		AddFoliageTypeCombo->SetIsOpen(false);
+	}
+
+	GWarn->BeginSlowTask(LOCTEXT("AddFoliageType_LoadPackage", "Fully Loading Package for Add"), true, false);
+	UObject* Asset = AssetData.GetAsset();
+	GWarn->EndSlowTask();
+
+	UFoliageType* FoliageType = FoliageEditMode->AddFoliageAsset(Asset);
+	if (FoliageType)
+	{
+		MeshTreeWidget->RequestTreeRefresh();
+	}
+}
+
 void SFoliageEdit::OnSetPaint()
 {
 	ClearAllToolSelection();
@@ -576,7 +693,7 @@ void SFoliageEdit::SetRadius(float InRadius)
 	FoliageEditMode->UISettings.SetRadius(InRadius);
 }
 
-float SFoliageEdit::GetRadius() const
+TOptional<float> SFoliageEdit::GetRadius() const
 {
 	return FoliageEditMode->UISettings.GetRadius();
 }
@@ -586,7 +703,7 @@ void SFoliageEdit::SetPaintDensity(float InDensity)
 	FoliageEditMode->UISettings.SetPaintDensity(InDensity);
 }
 
-float SFoliageEdit::GetPaintDensity() const
+TOptional<float> SFoliageEdit::GetPaintDensity() const
 {
 	return FoliageEditMode->UISettings.GetPaintDensity();
 }
@@ -597,11 +714,30 @@ void SFoliageEdit::SetEraseDensity(float InDensity)
 	FoliageEditMode->UISettings.SetUnpaintDensity(InDensity);
 }
 
-float SFoliageEdit::GetEraseDensity() const
+TOptional<float> SFoliageEdit::GetEraseDensity() const
 {
 	return FoliageEditMode->UISettings.GetUnpaintDensity();
 }
 
+EVisibility SFoliageEdit::GetAddFoliageTypeButtonTextVisibility() const
+{
+	return MeshTreeHeaderRowWidget.IsValid() && MeshTreeHeaderRowWidget->IsHovered() ? EVisibility::SelfHitTestInvisible : EVisibility::Collapsed;
+}
+
+TSharedRef<SWidget> SFoliageEdit::OnGetAddFoliageMeshAssetPicker()
+{
+	TArray<const UClass*> ClassFilters;
+	ClassFilters.Add(UFoliageType_InstancedStaticMesh::StaticClass());
+
+	return PropertyCustomizationHelpers::MakeAssetPickerWithMenu(FAssetData(),
+																false,
+																ClassFilters,
+																PropertyCustomizationHelpers::GetNewAssetFactoriesForClasses(ClassFilters),
+																FOnShouldFilterAsset(),
+																FOnAssetSelected::CreateSP(this, &SFoliageEdit::AddFoliageType),
+																FSimpleDelegate());
+}
+
 EVisibility SFoliageEdit::GetVisibility_EmptyList() const
 {
 	if (FoliageEditMode->GetFoliageMeshList().Num() > 0)
@@ -643,20 +779,12 @@ FReply SFoliageEdit::OnDrop_ListView(const FGeometry& MyGeometry, const FDragDro
 {
 	TArray<FAssetData> DroppedAssetData = AssetUtil::ExtractAssetDataFromDrag(DragDropEvent);
 	{
-		const FScopedTransaction Transaction(NSLOCTEXT("UnrealEd", "FoliageMode_AddMeshTransaction", "Foliage Editing: Add Mesh"));
-
 		for (int32 AssetIdx = 0; AssetIdx < DroppedAssetData.Num(); ++AssetIdx)
-		{
-			GWarn->BeginSlowTask(LOCTEXT("OnDrop_LoadPackage", "Fully Loading Package For Drop"), true, false);
-			UObject* Asset = DroppedAssetData[AssetIdx].GetAsset();
-			GWarn->EndSlowTask();
-
-			FoliageEditMode->AddFoliageAsset(Asset);
+		{	
+			AddFoliageType(DroppedAssetData[AssetIdx]);
 		}
 	}
 
-	MeshTreeWidget->RequestTreeRefresh();
-
 	return FReply::Handled();
 }
 TSharedPtr<SWidget> SFoliageEdit::ConstructFoliageMeshContextMenu() const
@@ -682,6 +810,44 @@ TSharedPtr<SWidget> SFoliageEdit::ConstructFoliageMeshContextMenu() const
 	return MenuBuilder.MakeWidget();
 }
 
+FText SFoliageEdit::GetFilterText() const
+{
+	FText TooltipText;
+	if (FoliageEditMode->UISettings.GetPaintToolSelected() || FoliageEditMode->UISettings.GetPaintBucketToolSelected())
+	{
+		TooltipText = LOCTEXT("PlacementFilter", "Placement Filter");
+	}
+	else if (FoliageEditMode->UISettings.GetReapplyToolSelected())
+	{
+		TooltipText = LOCTEXT("ReapplyFilter", "Reapply Filter");
+	}
+	else if (FoliageEditMode->UISettings.GetLassoSelectToolSelected())
+	{
+		TooltipText = LOCTEXT("SelectionFilter", "Selection Filter");
+	}
+
+	return TooltipText;
+}
+
+FText SFoliageEdit::GetFilterTooltipText() const
+{
+	FText TooltipText;
+	if (FoliageEditMode->UISettings.GetPaintToolSelected() || FoliageEditMode->UISettings.GetPaintBucketToolSelected())
+	{
+		TooltipText = LOCTEXT("PlacementFilter_Tooltip", "Limits the types of geometry upon which foliage instances can be placed.");
+	}
+	else if (FoliageEditMode->UISettings.GetReapplyToolSelected())
+	{
+		TooltipText = LOCTEXT("ReapplyFilter_Tooltip", "Limits the types of geometry upon which foliage instances can be reapplied.");
+	}
+	else if (FoliageEditMode->UISettings.GetLassoSelectToolSelected())
+	{
+		TooltipText = LOCTEXT("SelectionFilter_Tooltip", "Limits the types of geometry upon which foliage instances can be selected.");
+	}
+
+	return TooltipText;
+}
+
 void SFoliageEdit::OnCheckStateChanged_Landscape(ECheckBoxState InState)
 {
 	FoliageEditMode->UISettings.SetFilterLandscape(InState == ECheckBoxState::Checked ? true : false);
@@ -692,6 +858,25 @@ ECheckBoxState SFoliageEdit::GetCheckState_Landscape() const
 	return FoliageEditMode->UISettings.GetFilterLandscape() ? ECheckBoxState::Checked : ECheckBoxState::Unchecked;
 }
 
+FText SFoliageEdit::GetTooltipText_Landscape() const
+{
+	FText TooltipText;
+	if (FoliageEditMode->UISettings.GetPaintToolSelected() || FoliageEditMode->UISettings.GetPaintBucketToolSelected())
+	{
+		TooltipText = LOCTEXT("FilterLandscapeTooltip_Placement", "Place foliage on landscapes");
+	}
+	else if (FoliageEditMode->UISettings.GetReapplyToolSelected())
+	{
+		TooltipText = LOCTEXT("FilterLandscapeTooltip_Reapply", "Reapply to instances on landscapes");
+	}
+	else if (FoliageEditMode->UISettings.GetLassoSelectToolSelected())
+	{
+		TooltipText = LOCTEXT("FilterLandscapeTooltip_Select", "Select instances on landscapes");
+	}
+
+	return TooltipText;
+}
+
 void SFoliageEdit::OnCheckStateChanged_StaticMesh(ECheckBoxState InState)
 {
 	FoliageEditMode->UISettings.SetFilterStaticMesh(InState == ECheckBoxState::Checked ? true : false);
@@ -702,6 +887,25 @@ ECheckBoxState SFoliageEdit::GetCheckState_StaticMesh() const
 	return FoliageEditMode->UISettings.GetFilterStaticMesh() ? ECheckBoxState::Checked : ECheckBoxState::Unchecked;
 }
 
+FText SFoliageEdit::GetTooltipText_StaticMesh() const
+{
+	FText TooltipText;
+	if (FoliageEditMode->UISettings.GetPaintToolSelected() || FoliageEditMode->UISettings.GetPaintBucketToolSelected())
+	{
+		TooltipText = LOCTEXT("FilterStaticMeshTooltip_Placement", "Place foliage on static meshes");
+	}
+	else if (FoliageEditMode->UISettings.GetReapplyToolSelected())
+	{
+		TooltipText = LOCTEXT("FilterStaticMeshTooltip_Reapply", "Reapply to instances on static meshes");
+	}
+	else if (FoliageEditMode->UISettings.GetLassoSelectToolSelected())
+	{
+		TooltipText = LOCTEXT("FilterStaticMeshTooltip_Select", "Select instances on static meshes");
+	}
+
+	return TooltipText;
+}
+
 void SFoliageEdit::OnCheckStateChanged_BSP(ECheckBoxState InState)
 {
 	FoliageEditMode->UISettings.SetFilterBSP(InState == ECheckBoxState::Checked ? true : false);
@@ -712,6 +916,25 @@ ECheckBoxState SFoliageEdit::GetCheckState_BSP() const
 	return FoliageEditMode->UISettings.GetFilterBSP() ? ECheckBoxState::Checked : ECheckBoxState::Unchecked;
 }
 
+FText SFoliageEdit::GetTooltipText_BSP() const
+{
+	FText TooltipText;
+	if (FoliageEditMode->UISettings.GetPaintToolSelected() || FoliageEditMode->UISettings.GetPaintBucketToolSelected())
+	{
+		TooltipText = LOCTEXT("FilterBSPTooltip_Placement", "Place foliage on BSP");
+	}
+	else if (FoliageEditMode->UISettings.GetReapplyToolSelected())
+	{
+		TooltipText = LOCTEXT("FilterBSPTooltip_Reapply", "Reapply to instances on BSP");
+	}
+	else if (FoliageEditMode->UISettings.GetLassoSelectToolSelected())
+	{
+		TooltipText = LOCTEXT("FilterBSPTooltip_Select", "Select instances on BSP");
+	}
+
+	return TooltipText;
+}
+
 void SFoliageEdit::OnCheckStateChanged_Translucent(ECheckBoxState InState)
 {
 	FoliageEditMode->UISettings.SetFilterTranslucent(InState == ECheckBoxState::Checked ? true : false);
@@ -722,6 +945,25 @@ ECheckBoxState SFoliageEdit::GetCheckState_Translucent() const
 	return FoliageEditMode->UISettings.GetFilterTranslucent() ? ECheckBoxState::Checked : ECheckBoxState::Unchecked;
 }
 
+FText SFoliageEdit::GetTooltipText_Translucent() const
+{
+	FText TooltipText;
+	if (FoliageEditMode->UISettings.GetPaintToolSelected() || FoliageEditMode->UISettings.GetPaintBucketToolSelected())
+	{
+		TooltipText = LOCTEXT("FilterTranslucentTooltip_Placement", "Place foliage on translucent geometry");
+	}
+	else if (FoliageEditMode->UISettings.GetReapplyToolSelected())
+	{
+		TooltipText = LOCTEXT("FilterTranslucentTooltip_Reapply", "Reapply to instances on translucent geometry");
+	}
+	else if (FoliageEditMode->UISettings.GetLassoSelectToolSelected())
+	{
+		TooltipText = LOCTEXT("FilterTranslucentTooltip_Select", "Select instances on translucent geometry");
+	}
+
+	return TooltipText;
+}
+
 EVisibility SFoliageEdit::GetVisibility_Radius() const
 {
 	if (FoliageEditMode->UISettings.GetSelectToolSelected() || FoliageEditMode->UISettings.GetReapplyPaintBucketToolSelected() || FoliageEditMode->UISettings.GetPaintBucketToolSelected() )
@@ -734,7 +976,7 @@ EVisibility SFoliageEdit::GetVisibility_Radius() const
 
 EVisibility SFoliageEdit::GetVisibility_PaintDensity() const
 {
-	if (FoliageEditMode->UISettings.GetSelectToolSelected() || FoliageEditMode->UISettings.GetReapplyPaintBucketToolSelected() || FoliageEditMode->UISettings.GetLassoSelectToolSelected())
+	if (!FoliageEditMode->UISettings.GetPaintToolSelected())
 	{
 		return EVisibility::Collapsed;
 	}
@@ -744,7 +986,7 @@ EVisibility SFoliageEdit::GetVisibility_PaintDensity() const
 
 EVisibility SFoliageEdit::GetVisibility_EraseDensity() const
 {
-	if (FoliageEditMode->UISettings.GetSelectToolSelected() || FoliageEditMode->UISettings.GetReapplyPaintBucketToolSelected() || FoliageEditMode->UISettings.GetLassoSelectToolSelected() || FoliageEditMode->UISettings.GetPaintBucketToolSelected())
+	if (!FoliageEditMode->UISettings.GetPaintToolSelected())
 	{
 		return EVisibility::Collapsed;
 	}
@@ -776,15 +1018,7 @@ void SFoliageEdit::MeshTreeGetChildren(FFoliageMeshUIInfoPtr Item, TArray<FFolia
 
 void SFoliageEdit::MeshTreeOnSelectionChanged(FFoliageMeshUIInfoPtr Item, ESelectInfo::Type SelectInfo)
 {
-	TArray<UObject*> FoliageMeshes;
-	auto InfoList = MeshTreeWidget->GetSelectedItems();
-
-	for (const auto& Info : InfoList)
-	{
-		FoliageMeshes.Add(Info->Settings);
-	}
-
-	MeshDetailsWidget->SetObjects(FoliageMeshes, true);
+	RefreshMeshDetailsWidget();
 }
 
 void SFoliageEdit::OnRemoveFoliageType()
@@ -936,7 +1170,7 @@ void SFoliageEdit::OnCheckStateChanged_AllMeshes(ECheckBoxState InState)
 FText SFoliageEdit::GetMeshesHeaderText() const
 {
 	int32 NumMeshes = FoliageEditMode->GetFoliageMeshList().Num();
-	return FText::Format(LOCTEXT("FoliageMeshCount", "Meshes - {0}"), FText::AsNumber(NumMeshes));
+	return FText::Format(LOCTEXT("FoliageMeshCount", "Meshes ({0})"), FText::AsNumber(NumMeshes));
 }
 
 EColumnSortMode::Type SFoliageEdit::GetMeshColumnSortMode() const
@@ -949,4 +1183,23 @@ void SFoliageEdit::OnMeshesColumnSortModeChanged(EColumnSortPriority::Type InPri
 	FoliageEditMode->OnFoliageMeshListSortModeChanged(InSortMode);
 }
 
+void SFoliageEdit::HandleOnToolChanged()
+{
+	RefreshMeshDetailsWidget();
+}
+
+void SFoliageEdit::RefreshMeshDetailsWidget()
+{
+	TArray<UObject*> SelectedFoliageTypes;
+	auto SelectedFoliageInfoList = MeshTreeWidget->GetSelectedItems();
+
+	for (const auto& Info : SelectedFoliageInfoList)
+	{
+		SelectedFoliageTypes.Add(Info->Settings);
+	}
+
+	const bool bForceRefresh = true;
+	MeshDetailsWidget->SetObjects(SelectedFoliageTypes, bForceRefresh);
+}
+
 #undef LOCTEXT_NAMESPACE
diff --git a/Engine/Source/Editor/FoliageEdit/Private/SFoliageEdit.h b/Engine/Source/Editor/FoliageEdit/Private/SFoliageEdit.h
--- a/Engine/Source/Editor/FoliageEdit/Private/SFoliageEdit.h
+++ b/Engine/Source/Editor/FoliageEdit/Private/SFoliageEdit.h
@@ -19,8 +19,6 @@ public:
 	/** SCompoundWidget functions */
 	void Construct(const FArguments& InArgs);
 
-	~SFoliageEdit();
-
 	/** Does a full refresh on the list. */
 	void RefreshFullList();
 
@@ -43,6 +41,9 @@ private:
 	/** Creates the toolbar. */
 	TSharedRef<SWidget> BuildToolBar();
 
+	/** Adds the foliage type asset to the foliage actor's list of meshes. */
+	void AddFoliageType(const class FAssetData& AssetData);
+
 	/** Delegate callbacks for the UI */
 
 	/** Sets the tool mode to Paint. */
@@ -79,44 +80,68 @@ private:
 	void SetRadius(float InRadius);
 
 	/** Retrieves the brush Radius for the brush. */
-	float GetRadius() const;
+	TOptional<float> GetRadius() const;
 
 	/** Sets the Paint Density for the brush. */
 	void SetPaintDensity(float InDensity);
 
 	/** Retrieves the Paint Density for the brush. */
-	float GetPaintDensity() const;
+	TOptional<float> GetPaintDensity() const;
 
 	/** Sets the Erase Density for the brush. */
 	void SetEraseDensity(float InDensity);
 
 	/** Retrieves the Erase Density for the brush. */
-	float GetEraseDensity() const;
-	
+	TOptional<float> GetEraseDensity() const;
+
+	/** Gets the visibility of the Add Foliage Type text in the header row button */
+	EVisibility GetAddFoliageTypeButtonTextVisibility() const;
+
+	/** Gets the asset picker for adding a foliage type. */
+	TSharedRef<SWidget> OnGetAddFoliageMeshAssetPicker();
+
+	/** Retrieves the text for the filters option */
+	FText GetFilterText() const;
+
+	/** Retrieves the tooltip text for the filters option */
+	FText GetFilterTooltipText() const;
+
 	/** Sets the filter settings for if painting will occur on Landscapes. */
 	void OnCheckStateChanged_Landscape(ECheckBoxState InState);
 
 	/** Retrieves the filter settings for painting on Landscapes. */
 	ECheckBoxState GetCheckState_Landscape() const;
 
+	/** Retrieves the tooltip text for the landscape filter */
+	FText GetTooltipText_Landscape() const;
+
 	/** Sets the filter settings for if painting will occur on Static Meshes. */
 	void OnCheckStateChanged_StaticMesh(ECheckBoxState InState);
 
 	/** Retrieves the filter settings for painting on Static Meshes. */
 	ECheckBoxState GetCheckState_StaticMesh() const;
 
+	/** Retrieves the tooltip text for the static mesh filter */
+	FText GetTooltipText_StaticMesh() const;
+
 	/** Sets the filter settings for if painting will occur on BSPs. */
 	void OnCheckStateChanged_BSP(ECheckBoxState InState);
 
 	/** Retrieves the filter settings for painting on BSPs. */
 	ECheckBoxState GetCheckState_BSP() const;
 
+	/** Retrieves the tooltip text for the BSP filter */
+	FText GetTooltipText_BSP() const;
+
 	/** Sets the filter settings for if painting will occur on translucent meshes. */
 	void OnCheckStateChanged_Translucent(ECheckBoxState InState);
 
 	/** Retrieves the filter settings for painting on translucent meshes. */
 	ECheckBoxState GetCheckState_Translucent() const;
 
+	/** Retrieves the tooltip text for the translucent filter */
+	FText GetTooltipText_Translucent() const;
+
 	/** Checks if the text in the empty list overlay should appear. If the list is has items but the the drag and drop override is true, it will return EVisibility::Visible. */
 	EVisibility GetVisibility_EmptyList() const;
 
@@ -176,6 +201,12 @@ private:
 	EColumnSortMode::Type GetMeshColumnSortMode() const;
 	void OnMeshesColumnSortModeChanged(EColumnSortPriority::Type InPriority, const FName& InColumnName, EColumnSortMode::Type InSortMode);
 
+	/** Handler to trigger a refresh of the details view when the active tool changes */
+	void HandleOnToolChanged();
+
+	/** Refreshes the mesh details widget to match the current selection */
+	void RefreshMeshDetailsWidget();
+
 private:
 	/** Foliage mesh tree widget  */
 	TSharedPtr<SFoliageMeshTree>	MeshTreeWidget;
@@ -188,4 +219,10 @@ private:
 
 	/** Pointer to the foliage edit mode. */
 	FEdModeFoliage*					FoliageEditMode;
+
+	/** The Add Foliage Type combo button */
+	TSharedPtr<class SComboButton> AddFoliageTypeCombo;
+
+	/** The header row of the foliage mesh tree */
+	TSharedPtr<class SHeaderRow> MeshTreeHeaderRowWidget;
 };
diff --git a/Engine/Source/Editor/FoliageEdit/Private/SFoliageEditMeshRow.cpp b/Engine/Source/Editor/FoliageEdit/Private/SFoliageEditMeshRow.cpp
--- a/Engine/Source/Editor/FoliageEdit/Private/SFoliageEditMeshRow.cpp
+++ b/Engine/Source/Editor/FoliageEdit/Private/SFoliageEditMeshRow.cpp
@@ -103,6 +103,7 @@ TSharedRef<SWidget> SFoliageEditMeshRow::GenerateWidgetForColumn(const FName& Co
 			SNew(SHorizontalBox)
 			
 			+SHorizontalBox::Slot()
+			.Padding(FMargin(10,1,0,1))
 			.VAlign(VAlign_Center)
 			.AutoWidth()
 			[
diff --git a/Engine/Source/Runtime/Foliage/Private/InstancedFoliage.cpp b/Engine/Source/Runtime/Foliage/Private/InstancedFoliage.cpp
--- a/Engine/Source/Runtime/Foliage/Private/InstancedFoliage.cpp
+++ b/Engine/Source/Runtime/Foliage/Private/InstancedFoliage.cpp
@@ -281,7 +281,7 @@ UFoliageType::UFoliageType(const FObjectInitializer& ObjectInitializer)
 	MinimumLayerWeight = 0.5f;
 	DisplayOrder = 0;
 	IsSelected = false;
-	ReapplyDensityAmount = 1.0f;
+	DensityAdjustmentFactor = 1.0f;
 	CollisionWithWorld = false;
 	CollisionScale = FVector(0.9f, 0.9f, 0.9f);
 	VertexColorMask = FOLIAGEVERTEXCOLORMASK_Disabled;
diff --git a/Engine/Source/Runtime/Foliage/Public/FoliageType.h b/Engine/Source/Runtime/Foliage/Public/FoliageType.h
--- a/Engine/Source/Runtime/Foliage/Public/FoliageType.h
+++ b/Engine/Source/Runtime/Foliage/Public/FoliageType.h
@@ -4,6 +4,7 @@
 #include "Curves/CurveFloat.h"
 #include "FoliageType.generated.h"
 
+
 UENUM()
 enum FoliageVertexColorMask
 {
@@ -39,82 +40,126 @@
 {
 	GENERATED_UCLASS_BODY()
 
+	/* Gets/Sets the mesh associated with this FoliageType */
+	virtual UStaticMesh* GetStaticMesh() const PURE_VIRTUAL(UFoliageType::GetStaticMesh, return nullptr; );
+	virtual void SetStaticMesh(UStaticMesh* InStaticMesh) PURE_VIRTUAL(UFoliageType::SetStaticMesh,);
+
+	virtual void Serialize(FArchive& Ar) override;
+
+	FOLIAGE_API FVector GetRandomScale() const;
+
+#if WITH_EDITOR
+	/* Lets subclasses decide if the InstancedFoliageActor should reallocate its instances if the specified property change event occurs */
+	virtual bool IsFoliageReallocationRequiredForPropertyChange(struct FPropertyChangedEvent& PropertyChangedEvent) const { return true; }
+
+	virtual void PostEditChangeProperty(struct FPropertyChangedEvent& PropertyChangedEvent) override;
+#endif
+
 	/* A GUID that is updated every time the foliage type is modified, 
 	   so foliage placed in the level can detect the FoliageType has changed. */
 	UPROPERTY()
 	FGuid UpdateGuid;
 
+
+public:
+	// PAINTING
+
 	/** Foliage instances will be placed at this density, specified in instances per 1000x1000 unit area */
 	UPROPERTY(EditAnywhere, Category=Painting, meta=(DisplayName="Density / 1Kuu", UIMin=0, ClampMin=0))
 	float Density;
 
+	/** The factor by which to adjust the density of instances. Values >1 will increase density while values <1 will decrease it. */
+	UPROPERTY(EditAnywhere, Category=Painting, meta=(UIMin=0, ClampMin=0, ReapplyCondition="ReapplyDensity"))
+	float DensityAdjustmentFactor;
+
 	/** The minimum distance between foliage instances */
-	UPROPERTY(EditAnywhere, Category=Painting, meta=(UIMin=0, ClampMin=0))
+	UPROPERTY(EditAnywhere, Category=Painting, meta=(UIMin=0, ClampMin=0, ReapplyCondition="ReapplyRadius"))
 	float Radius;
 
-	/** Specified foliage instance scaling behavior */
-	UPROPERTY(EditAnywhere, Category=Painting)
+	/** 
+	 *  When painting on static meshes, foliage instance placement can be limited to areas where the static mesh has values in the selected vertex color channel(s). 
+	 *  This allows a static mesh to mask out certain areas to prevent foliage from being placed there
+	 */
+	UPROPERTY(EditAnywhere, AdvancedDisplay, Category=Painting, meta=(ReapplyCondition="ReapplyVertexColorMask"))
+	TEnumAsByte<enum FoliageVertexColorMask> VertexColorMask;
+
+	/** Specifies the threshold value above which the static mesh vertex color value must be, in order for foliage instances to be placed in a specific area */
+	UPROPERTY(EditAnywhere, AdvancedDisplay, Category=Painting, meta=(HideBehind="VertexColorMask"))
+	float VertexColorMaskThreshold;
+
+	/** 
+	 *  When unchecked, foliage instances will be placed only when the vertex color in the specified channel(s) is above the threshold amount. 
+	 *  When checked, the vertex color must be less than the threshold amount 
+	 */
+	UPROPERTY(EditAnywhere, AdvancedDisplay, Category=Painting, meta=(HideBehind="VertexColorMask"))
+	uint32 VertexColorMaskInvert:1;
+
+public:
+	// PLACEMENT
+
+	/** Specifies foliage instance scaling behavior when painting. */
+	UPROPERTY(EditAnywhere, Category=Placement, meta=(ReapplyCondition="ReapplyScaling"))
 	EFoliageScaling Scaling;
 
 	/** Specifies the range of scale, from minimum to maximum, to apply to a foliage instance's X Scale property */
-	UPROPERTY(EditAnywhere, Category=Painting, meta=(UIMin=0.001))
+	UPROPERTY(EditAnywhere, Category=Placement, meta=(ClampMin="0.001", UIMin="0.001", ReapplyCondition="ReapplyScaleX"))
 	FFloatInterval ScaleX;
 
 	/** Specifies the range of scale, from minimum to maximum, to apply to a foliage instance's Y Scale property */
-	UPROPERTY(EditAnywhere, Category=Painting, meta=(UIMin=0.001))
+	UPROPERTY(EditAnywhere, Category=Placement, meta=(ClampMin="0.001", UIMin="0.001", ReapplyCondition="ReapplyScaleY"))
 	FFloatInterval ScaleY;
 
 	/** Specifies the range of scale, from minimum to maximum, to apply to a foliage instance's Z Scale property */
-	UPROPERTY(EditAnywhere, Category=Painting, meta=(UIMin=0.001))
+	UPROPERTY(EditAnywhere, Category=Placement, meta=(ClampMin="0.001", UIMin="0.001", ReapplyCondition="ReapplyScaleZ"))
 	FFloatInterval ScaleZ;
 
+	/** Specifies a range from minimum to maximum of the offset to apply to a foliage instance's Z location */
+	UPROPERTY(EditAnywhere, Category=Placement, meta=(DisplayName="Z Offset", ReapplyCondition="ReapplyZOffset"))
+	FFloatInterval ZOffset;
+
 	/** Whether foliage instances should have their angle adjusted away from vertical to match the normal of the surface they're painted on */
-	UPROPERTY(EditAnywhere, Category=Placement)
+	UPROPERTY(EditAnywhere, Category=Placement, meta=(ReapplyCondition="ReapplyAlignToNormal"))
 	uint32 AlignToNormal:1;
 
 	/** The maximum angle in degrees that foliage instances will be adjusted away from the vertical */
-	UPROPERTY(EditAnywhere, Category=Placement)
+	UPROPERTY(EditAnywhere, Category=Placement, meta=(HideBehind="AlignToNormal"))
 	float AlignMaxAngle;
 
+	/** If selected, foliage instances will have a random yaw rotation around their vertical axis applied */
+	UPROPERTY(EditAnywhere, Category=Placement, meta=(ReapplyCondition="ReapplyRandomYaw"))
+	uint32 RandomYaw:1;
+
 	/** A random pitch adjustment can be applied to each instance, up to the specified angle in degrees, from the original vertical */
-	UPROPERTY(EditAnywhere, Category=Placement)
+	UPROPERTY(EditAnywhere, Category=Placement, meta=(ReapplyCondition="ReapplyRandomPitchAngle"))
 	float RandomPitchAngle;
 
 	/* Foliage instances will only be placed on surfaces sloping in the specified angle range from the horizontal */
-	UPROPERTY(EditAnywhere, Category=Placement, meta=(UIMin=0))
+	UPROPERTY(EditAnywhere, Category=Placement, meta=(UIMin=0, ReapplyCondition="ReapplyGroundSlope"))
 	FFloatInterval GroundSlopeAngle;
 
 	/* The valid altitude range where foliage instances will be placed, specified using minimum and maximum world coordinate Z values */
-	UPROPERTY(EditAnywhere, Category=Placement)
+	UPROPERTY(EditAnywhere, Category=Placement, meta=(ReapplyCondition="ReapplyHeight"))
 	FFloatInterval Height;
 
 	/** If a layer name is specified, painting on landscape will limit the foliage to areas of landscape with the specified layer painted */
-	UPROPERTY(EditAnywhere, AdvancedDisplay, Category=Placement)
+	UPROPERTY(EditAnywhere, AdvancedDisplay, Category=Placement, meta=(ReapplyCondition="ReapplyLandscapeLayers"))
 	TArray<FName> LandscapeLayers;
 
 	UPROPERTY()
 	FName LandscapeLayer_DEPRECATED;
 	
-	/** Specifies the minimum value above which the landscape layer weight value must be, in order for foliage instances to be placed in a specific area */
-	UPROPERTY(EditAnywhere, AdvancedDisplay, Category=Placement)
-	float MinimumLayerWeight;
-
-	/** If selected, foliage instances will have a random yaw rotation around their vertical axis applied */
-	UPROPERTY(EditAnywhere, Category=Placement)
-	uint32 RandomYaw:1;
-
-	/** Specifies a range from minimum to maximum of the offset to apply to a foliage instance's Z location */
-	UPROPERTY(EditAnywhere, Category=Painting, meta=(DisplayName="Z Offset"))
-	FFloatInterval ZOffset;
-	
 	/* If checked, an overlap test with existing world geometry is performed before each instance is placed */
-	UPROPERTY(EditAnywhere, AdvancedDisplay, Category=Painting)
+	UPROPERTY(EditAnywhere, AdvancedDisplay, Category=Placement, meta=(ReapplyCondition="ReapplyCollisionWithWorld"))
 	uint32 CollisionWithWorld:1;
 
 	/* The foliage instance's collision bounding box will be scaled by the specified amount before performing the overlap check */
-	UPROPERTY(EditAnywhere, AdvancedDisplay, Category=Painting)
+	UPROPERTY(EditAnywhere, AdvancedDisplay, Category=Placement, meta=(HideBehind="CollisionWithWorld"))
 	FVector CollisionScale;
 
+	/** Specifies the minimum value above which the landscape layer weight value must be, in order for foliage instances to be placed in a specific area */
+	UPROPERTY(EditAnywhere, AdvancedDisplay, Category=Placement, meta=(HideBehind="LandscapeLayers"))
+	float MinimumLayerWeight;
+
 	UPROPERTY()
 	FBoxSphereBounds MeshBounds;
 
@@ -122,82 +167,9 @@
 	UPROPERTY()
 	FVector LowBoundOriginRadius;
 
-	/** 
-	 *  When painting on static meshes, foliage instance placement can be limited to areas where the static mesh has values in the selected vertex color channel(s). 
-	 *  This allows a static mesh to mask out certain areas to prevent foliage from being placed there
-	 */
-	UPROPERTY(EditAnywhere, AdvancedDisplay, Category=Painting)
-	TEnumAsByte<enum FoliageVertexColorMask> VertexColorMask;
 
-	/** Specifies the threshold value above which the static mesh vertex color value must be, in order for foliage instances to be placed in a specific area */
-	UPROPERTY(EditAnywhere, AdvancedDisplay, Category=Painting)
-	float VertexColorMaskThreshold;
+public:
+	// INSTANCE SETTINGS
 
-	/** 
-	 *  When unchecked, foliage instances will be placed only when the vertex color in the specified channel(s) is above the threshold amount. 
-	 *  When checked, the vertex color must be less than the threshold amount 
-	 */
-	UPROPERTY(EditAnywhere, AdvancedDisplay, Category=Painting)
-	uint32 VertexColorMaskInvert:1;
-
-	UPROPERTY(EditAnywhere, Category=Painting, meta=(DisplayName="Density Amount", UIMin=0, ClampMin=0))
-	float ReapplyDensityAmount;
-
-	/** If checked, the density of foliage instances already placed will be adjusted. <1 will remove instances to reduce the density, >1 will place extra instances */
-	UPROPERTY(EditAnywhere, Category=Painting)
-	uint32 ReapplyDensity:1;
-
-	/** If checked, foliage instances not meeting the new Radius constraint will be removed */
-	UPROPERTY(EditAnywhere, Category=Painting)
-	uint32 ReapplyRadius:1;
-
-	/** If checked, foliage instances will have their normal alignment adjusted by the Reapply tool */
-	UPROPERTY(EditAnywhere, Category=Painting)
-	uint32 ReapplyAlignToNormal:1;
-
-	/** If checked, foliage instances will have their yaw adjusted by the Reapply tool */
-	UPROPERTY(EditAnywhere, Category=Painting)
-	uint32 ReapplyRandomYaw:1;
-
-	/** If checked, foliage instances will have their X scale adjusted by the Reapply tool */
-	UPROPERTY(EditAnywhere, Category=Painting)
-	uint32 ReapplyScaleX:1;
-
-	/** If checked, foliage instances will have their Y scale adjusted by the Reapply tool */
-	UPROPERTY(EditAnywhere, Category=Painting)
-	uint32 ReapplyScaleY:1;
-
-	/** If checked, foliage instances will have their Z scale adjusted by the Reapply tool */
-	UPROPERTY(EditAnywhere, Category=Painting)
-	uint32 ReapplyScaleZ:1;
-
-	/** If checked, foliage instances will have their pitch adjusted by the Reapply tool */
-	UPROPERTY(EditAnywhere, Category=Painting)
-	uint32 ReapplyRandomPitchAngle:1;
-
-	/** If checked, foliage instances not meeting the ground slope condition will be removed by the Reapply too */
-	UPROPERTY(EditAnywhere, Category=Painting)
-	uint32 ReapplyGroundSlope:1;
-
-	/* If checked, foliage instances not meeting the valid Z height condition will be removed by the Reapply tool */
-	UPROPERTY(EditAnywhere, Category=Painting)
-	uint32 ReapplyHeight:1;
-
-	/* If checked, foliage instances painted on areas that do not have the appopriate landscape layer painted will be removed by the Reapply tool */
-	UPROPERTY(EditAnywhere, Category=Painting)
-	uint32 ReapplyLandscapeLayer:1;
-
-	/** If checked, foliage instances will have their Z offset adjusted by the Reapply tool */
-	UPROPERTY(EditAnywhere, Category=Painting)
-	uint32 ReapplyZOffset:1;
-
-	/* If checked, foliage instances will have an overlap test with the world reapplied, and overlapping instances will be removed by the Reapply tool */
-	UPROPERTY(EditAnywhere, Category=Painting)
-	uint32 ReapplyCollisionWithWorld:1;
-
-	/* If checked, foliage instances no longer matching the vertex color constraint will be removed by the Reapply too */
-	UPROPERTY(EditAnywhere, Category=Painting)
-	uint32 ReapplyVertexColorMask:1;
-
 	/**
 	 * The distance where instances will begin to fade out if using a PerInstanceFadeAmount material node. 0 disables.
@@ -219,7 +191,7 @@
 	uint32 bAffectDynamicIndirectLighting:1;
 
 	/** Controls whether the primitive should affect dynamic distance field lighting methods.  This flag is only used if CastShadow is true. */
-	UPROPERTY(EditAnywhere, BlueprintReadOnly, Category=InstanceSettings, meta = (EditCondition = "CastShadow"))
+	UPROPERTY(EditAnywhere, BlueprintReadOnly, Category=InstanceSettings, meta=(EditCondition="CastShadow"))
 	uint32 bAffectDistanceFieldLighting:1;
 
 	/** Controls whether the foliage should cast shadows in the case of non precomputed shadowing.  This flag is only used if CastShadow is true. */
@@ -243,15 +215,15 @@
 	uint32 bOverrideLightMapRes:1;
 
 	/** Overrides the lightmap resolution defined in the static mesh */
-	UPROPERTY(EditAnywhere, BlueprintReadOnly, Category=InstanceSettings, meta = (DisplayName = "Light Map Resolution", EditCondition = "bOverrideLightMapRes"))
+	UPROPERTY(EditAnywhere, BlueprintReadOnly, Category=InstanceSettings, meta=(DisplayName="Light Map Resolution", EditCondition="bOverrideLightMapRes"))
 	int32 OverriddenLightMapRes;
 
 	/** Custom collision for foliage */
-	UPROPERTY(EditAnywhere, Category=InstanceSettings, meta = (HideObjectType = true))
+	UPROPERTY(EditAnywhere, Category=InstanceSettings, meta=(HideObjectType=true))
 	struct FBodyInstance BodyInstance;
 
 	/** Force navmesh */
-	UPROPERTY(EditAnywhere, Category=InstanceSettings, meta = (HideObjectType = true))
+	UPROPERTY(EditAnywhere, Category=InstanceSettings, meta=(HideObjectType=true))
 	TEnumAsByte<EHasCustomNavigableGeometry::Type> CustomNavigableGeometry;
 
 	UPROPERTY()
@@ -260,97 +232,151 @@
 	UPROPERTY()
 	uint32 IsSelected:1;
 
-	/* Gets/Sets the mesh associated with this FoliageType */
-	virtual UStaticMesh* GetStaticMesh() const PURE_VIRTUAL(UFoliageType::GetStaticMesh, return nullptr; );
-	virtual void SetStaticMesh(UStaticMesh* InStaticMesh) PURE_VIRTUAL(UFoliageType::SetStaticMesh,);
+public:
+	// PROCEDURAL
 
-#if WITH_EDITOR
-	/* Lets subclasses decide if the InstancedFoliageActor should reallocate its instances if the specified property change event occurs */
-	virtual bool IsFoliageReallocationRequiredForPropertyChange(struct FPropertyChangedEvent& PropertyChangedEvent) const { return true; }
+	FOLIAGE_API float GetSeedDensitySquared() const { return InitialSeedDensity * InitialSeedDensity; }
+	FOLIAGE_API float GetMaxRadius() const;
+	FOLIAGE_API float GetScaleForAge(const float Age) const;
+	FOLIAGE_API float GetInitAge(FRandomStream& RandomStream) const;
+	FOLIAGE_API float GetNextAge(const float CurrentAge, const int32 NumSteps) const;
 
-	virtual void PostEditChangeProperty(struct FPropertyChangedEvent& PropertyChangedEvent) override;
-#endif
+	// COLLISION
 
-	/** Procedural specific parameters */
-
-	/** Specifies the number of seeds to populate along 10 meters. The number is implicitly squared to cover a 10m x 10m area*/
-	UPROPERTY(Category = Procedural, EditAnywhere, meta = (Subcategory = "Clustering", ClampMin = "0.0", UIMin = "0.0"))
-	float InitialSeedDensity;
-
-	/** The average distance between the spreading instance and its seeds. For example, a tree with an AverageSpreadDistance 10 will ensure the average distance between the tree and its seeds is 10cm */
-	UPROPERTY(Category = Procedural, EditAnywhere, meta = (Subcategory = "Clustering", ClampMin = "0.0", UIMin = "0.0"))
-	float AverageSpreadDistance;
-
-	/** Specifies how much seed distance varies from the average. For example, a tree with an AverageSpreadDistance 10 and a SpreadVariance 1 will produce seeds with an average distance of 10cm plus or minus 1cm */
-	UPROPERTY(Category = Procedural, EditAnywhere, meta = (Subcategory = "Clustering", ClampMin = "0.0", UIMin = "0.0"))
-	float SpreadVariance;
-
 	/** The CollisionRadius determines when two instances overlap. When two instances overlap a winner will be picked based on rules and priority. */
-	UPROPERTY(Category = Procedural, EditAnywhere, meta = (Subcategory = "Collision", ClampMin = "0.0", UIMin = "0.0"))
+	UPROPERTY(Category=Procedural, EditAnywhere, meta=(Subcategory="Collision", ClampMin="0.0", UIMin="0.0"))
 	float CollisionRadius;
 
 	/** The ShadeRadius determines when two instances overlap. If an instance can grow in the shade this radius is ignored.*/
-	UPROPERTY(Category = Procedural, EditAnywhere, meta = (Subcategory = "Collision", ClampMin = "0.0", UIMin = "0.0"))
+	UPROPERTY(Category=Procedural, EditAnywhere, meta=(Subcategory="Collision", ClampMin="0.0", UIMin="0.0"))
 	float ShadeRadius;
 
 	/** The number of times we age the species and spread its seeds. */
-	UPROPERTY(Category = Procedural, EditAnywhere, meta = (Subcategory = "Clustering", ClampMin = "0", UIMin = "0"))
+	UPROPERTY(Category=Procedural, EditAnywhere, meta=(Subcategory="Clustering", ClampMin="0", UIMin="0"))
 	int32 NumSteps;
 
+	// CLUSTERING
+
+	/** Specifies the number of seeds to populate along 10 meters. The number is implicitly squared to cover a 10m x 10m area*/
+	UPROPERTY(Category=Procedural, EditAnywhere, meta=(Subcategory="Clustering", ClampMin="0.0", UIMin="0.0"))
+	float InitialSeedDensity;
+
+	/** The average distance between the spreading instance and its seeds. For example, a tree with an AverageSpreadDistance 10 will ensure the average distance between the tree and its seeds is 10cm */
+	UPROPERTY(Category=Procedural, EditAnywhere, meta=(Subcategory="Clustering", ClampMin="0.0", UIMin="0.0"))
+	float AverageSpreadDistance;
+
+	/** Specifies how much seed distance varies from the average. For example, a tree with an AverageSpreadDistance 10 and a SpreadVariance 1 will produce seeds with an average distance of 10cm plus or minus 1cm */
+	UPROPERTY(Category=Procedural, EditAnywhere, meta=(Subcategory="Clustering", ClampMin="0.0", UIMin="0.0"))
+	float SpreadVariance;
+
 	/** The number of seeds an instance will spread in a single step of the simulation. */
-	UPROPERTY(Category = Procedural, EditAnywhere, meta = (Subcategory = "Clustering", ClampMin = "0", UIMin = "0"))
+	UPROPERTY(Category=Procedural, EditAnywhere, meta=(Subcategory="Clustering", ClampMin="0", UIMin="0"))
 	int32 SeedsPerStep;
 
 	/** The seed that determines placement of initial seeds. */
-	UPROPERTY(Category = Procedural, EditAnywhere, meta = (Subcategory = "Clustering"))
+	UPROPERTY(Category=Procedural, EditAnywhere, meta=(Subcategory="Clustering"))
 	int32 DistributionSeed;
 
 	/** The seed that determines placement of initial seeds. */
-	UPROPERTY(Category = Procedural, EditAnywhere, meta = (Subcategory = "Clustering"))
+	UPROPERTY(Category=Procedural, EditAnywhere, meta=(Subcategory="Clustering"))
 	float MaxInitialSeedOffset;
 
+	// GROWTH
+
 	/** Whether the species can grow in shade. If this is true shade radius is ignored during overlap tests*/
-	UPROPERTY(Category = Procedural, EditAnywhere, meta = (Subcategory = "Growth"))
+	UPROPERTY(Category=Procedural, EditAnywhere, meta=(Subcategory="Growth"))
 	bool bGrowsInShade;
 
-	/** The minimum scale that an instance will be. Corresponds to age = 0*/
-	UPROPERTY(Category = Procedural, EditAnywhere, meta = (Subcategory = "Growth", ClampMin = "0.0", UIMin = "0.0"))
+	/** The minimum scale that an instance will be. Corresponds to age = 0. */
+	UPROPERTY(Category=Procedural, EditAnywhere, meta=(Subcategory="Growth", ClampMin="0.0", UIMin="0.0"))
 	float MinScale;
 
-	/** The maximum scale that a seed will grow to. Corresponds to age = 1*/
-	UPROPERTY(Category = Procedural, EditAnywhere, meta = (Subcategory = "Growth", ClampMin = "0.0", UIMin = "0.0"))
+	/** The maximum scale that a seed will grow to. Corresponds to age = 10. */
+	UPROPERTY(Category=Procedural, EditAnywhere, meta=(Subcategory="Growth", ClampMin="0.0", UIMin="0.0"))
 	float MaxScale;
 
 	/** Specifies the oldest a new seed can be. The new seed will have an age randomly distributed in [0,InitMaxAge] */
-	UPROPERTY(Category = Procedural, EditAnywhere, meta = (Subcategory = "Growth", ClampMin = "0.0", UIMin = "0.0"))
+	UPROPERTY(Category=Procedural, EditAnywhere, meta=(Subcategory="Growth", ClampMin="0.0", UIMin="0.0"))
 	float InitialMaxAge;
 
 	/** Specifies the oldest a seed can be. After reaching this age the instance will still spread seeds, but will not get any older*/
-	UPROPERTY(Category = Procedural, EditAnywhere, meta = (Subcategory = "Growth", ClampMin = "0.0", UIMin = "0.0"))
+	UPROPERTY(Category=Procedural, EditAnywhere, meta=(Subcategory="Growth", ClampMin="0.0", UIMin="0.0"))
 	float MaxAge;
 
 	/** When two instances overlap we must determine which instance to remove. The instance with a lower OverlapPriority will be removed. In the case where OverlapPriority is the same regular simulation rules apply.*/
-	UPROPERTY(Category = Procedural, EditAnywhere, meta = (Subcategory = "Growth"))
+	UPROPERTY(Category=Procedural, EditAnywhere, meta=(Subcategory="Growth"))
 	float OverlapPriority;
 
 	UPROPERTY()
 	int32 ChangeCount;
 
-	
-	FOLIAGE_API float GetSeedDensitySquared() const { return InitialSeedDensity * InitialSeedDensity; }
-	FOLIAGE_API float GetMaxRadius() const;
-	FOLIAGE_API float GetScaleForAge(const float Age) const;
-	FOLIAGE_API float GetInitAge(FRandomStream& RandomStream) const;
-	FOLIAGE_API float GetNextAge(const float CurrentAge, const int32 NumSteps) const;
+public:
+	// REAPPLY EDIT CONDITIONS
 
-	virtual void Serialize(FArchive& Ar) override;
+	/** If checked, the density of foliage instances already placed will be adjusted by the density adjustment factor. */
+	UPROPERTY(EditDefaultsOnly, Category=Reapply)
+	uint32 ReapplyDensity:1;
 
-	FOLIAGE_API FVector GetRandomScale() const;
+	/** If checked, foliage instances not meeting the new Radius constraint will be removed */
+	UPROPERTY(EditDefaultsOnly, Category=Reapply)
+	uint32 ReapplyRadius:1;
 
+	/** If checked, foliage instances will have their normal alignment adjusted by the Reapply tool */
+	UPROPERTY(EditDefaultsOnly, Category=Reapply)
+	uint32 ReapplyAlignToNormal:1;
+
+	/** If checked, foliage instances will have their yaw adjusted by the Reapply tool */
+	UPROPERTY(EditDefaultsOnly, Category=Reapply)
+	uint32 ReapplyRandomYaw:1;
+
+	/** If checked, foliage instances will have their scale adjusted to fit the specified scaling behavior by the Reapply tool */
+	UPROPERTY(EditDefaultsOnly, Category=Reapply)
+	uint32 ReapplyScaling:1;
+
+	/** If checked, foliage instances will have their X scale adjusted by the Reapply tool */
+	UPROPERTY(EditDefaultsOnly, Category=Reapply)
+	uint32 ReapplyScaleX:1;
+
+	/** If checked, foliage instances will have their Y scale adjusted by the Reapply tool */
+	UPROPERTY(EditDefaultsOnly, Category=Reapply)
+	uint32 ReapplyScaleY:1;
+
+	/** If checked, foliage instances will have their Z scale adjusted by the Reapply tool */
+	UPROPERTY(EditDefaultsOnly, Category=Reapply)
+	uint32 ReapplyScaleZ:1;
+
+	/** If checked, foliage instances will have their pitch adjusted by the Reapply tool */
+	UPROPERTY(EditDefaultsOnly, Category=Reapply)
+	uint32 ReapplyRandomPitchAngle:1;
+
+	/** If checked, foliage instances not meeting the ground slope condition will be removed by the Reapply too */
+	UPROPERTY(EditDefaultsOnly, Category=Reapply)
+	uint32 ReapplyGroundSlope:1;
+
+	/* If checked, foliage instances not meeting the valid Z height condition will be removed by the Reapply tool */
+	UPROPERTY(EditDefaultsOnly, Category=Reapply)
+	uint32 ReapplyHeight:1;
+
+	/* If checked, foliage instances painted on areas that do not have the appropriate landscape layer painted will be removed by the Reapply tool */
+	UPROPERTY(EditDefaultsOnly, Category=Reapply)
+	uint32 ReapplyLandscapeLayers:1;
+
+	/** If checked, foliage instances will have their Z offset adjusted by the Reapply tool */
+	UPROPERTY(EditDefaultsOnly, Category=Reapply)
+	uint32 ReapplyZOffset:1;
+
+	/* If checked, foliage instances will have an overlap test with the world reapplied, and overlapping instances will be removed by the Reapply tool */
+	UPROPERTY(EditDefaultsOnly, Category=Reapply)
+	uint32 ReapplyCollisionWithWorld:1;
+
+	/* If checked, foliage instances no longer matching the vertex color constraint will be removed by the Reapply too */
+	UPROPERTY(EditDefaultsOnly, Category=Reapply)
+	uint32 ReapplyVertexColorMask:1;
+
 private:
 
 	/** The curve used to interpolate the instance scale.*/
-	UPROPERTY(Category = Procedural, EditAnywhere, meta = (Subcategory = "Growth"))
+	UPROPERTY(Category = Procedural, EditAnywhere, meta = (Subcategory="Growth"))
 	FRuntimeFloatCurve ScaleCurve;
 	
 	// Deprecated since FFoliageCustomVersion::FoliageTypeCustomization
