Non-unity compilation fixes #UE4

diff --git a/Engine/Source/Editor/BlueprintGraph/Private/BlueprintActionDatabase.cpp b/Engine/Source/Editor/BlueprintGraph/Private/BlueprintActionDatabase.cpp
--- a/Engine/Source/Editor/BlueprintGraph/Private/BlueprintActionDatabase.cpp
+++ b/Engine/Source/Editor/BlueprintGraph/Private/BlueprintActionDatabase.cpp
@@ -20,6 +20,7 @@
 #include "BlueprintEditorUtils.h"	// for FindBlueprintForGraph()
 #include "BlueprintActionDatabaseRegistrar.h"
 #include "BlueprintActionFilter.h"	// for FBlueprintActionContext
+#include "Engine/LevelScriptBlueprint.h" // for ULevelScriptBlueprint
 
 // used below in FBlueprintNodeSpawnerFactory::MakeMacroNodeSpawner()
 #include "K2Node_MacroInstance.h"
diff --git a/Engine/Source/Editor/FoliageEdit/Private/FoliageTypeCustomizationHelpers.cpp b/Engine/Source/Editor/FoliageEdit/Private/FoliageTypeCustomizationHelpers.cpp
--- a/Engine/Source/Editor/FoliageEdit/Private/FoliageTypeCustomizationHelpers.cpp
+++ b/Engine/Source/Editor/FoliageEdit/Private/FoliageTypeCustomizationHelpers.cpp
@@ -2,6 +2,11 @@
 
 #include "UnrealEd.h"
 #include "FoliageTypeCustomizationHelpers.h"
+#include "FoliageType.h"
+#include "IDetailPropertyRow.h"
+#include "DetailLayoutBuilder.h"
+#include "DetailCategoryBuilder.h"
+#include "PropertyHandle.h"
 
 void FFoliageTypeCustomizationHelpers::ModifyFoliagePropertyRow(IDetailPropertyRow* PropertyRow, const TAttribute<EVisibility>& InVisibility, const TAttribute<bool>& InEnabled)
 {
diff --git a/Engine/Source/Editor/FoliageEdit/Private/FoliageTypeCustomizationHelpers.h b/Engine/Source/Editor/FoliageEdit/Private/FoliageTypeCustomizationHelpers.h
--- a/Engine/Source/Editor/FoliageEdit/Private/FoliageTypeCustomizationHelpers.h
+++ b/Engine/Source/Editor/FoliageEdit/Private/FoliageTypeCustomizationHelpers.h
@@ -8,6 +8,7 @@ class FEdModeFoliage;
 class IDetailLayoutBuilder;
 class IDetailCategoryBuilder;
 class IDetailPropertyRow;
+class IPropertyHandle;
 
 class FFoliageTypeCustomizationHelpers
 {
