[UETOOL-266] - Procedural foliage assets & properties are hidden behind an experimental flag

diff --git a/Engine/Source/Developer/AssetTools/Private/AssetTypeActions/AssetTypeActions_ProceduralFoliageSpawner.cpp b/Engine/Source/Developer/AssetTools/Private/AssetTypeActions/AssetTypeActions_ProceduralFoliageSpawner.cpp
--- a/Engine/Source/Developer/AssetTools/Private/AssetTypeActions/AssetTypeActions_ProceduralFoliageSpawner.cpp
+++ b/Engine/Source/Developer/AssetTools/Private/AssetTypeActions/AssetTypeActions_ProceduralFoliageSpawner.cpp
@@ -6,4 +6,9 @@
 UClass* FAssetTypeActions_ProceduralFoliageSpawner::GetSupportedClass() const
 {
 	return UProceduralFoliageSpawner::StaticClass();
+}
+
+bool FAssetTypeActions_ProceduralFoliageSpawner::CanFilter()
+{
+	return GetDefault<UEditorExperimentalSettings>()->bProceduralFoliage;
 }
\ No newline at end of file
diff --git a/Engine/Source/Developer/AssetTools/Private/AssetTypeActions/AssetTypeActions_ProceduralFoliageSpawner.h b/Engine/Source/Developer/AssetTools/Private/AssetTypeActions/AssetTypeActions_ProceduralFoliageSpawner.h
--- a/Engine/Source/Developer/AssetTools/Private/AssetTypeActions/AssetTypeActions_ProceduralFoliageSpawner.h
+++ b/Engine/Source/Developer/AssetTools/Private/AssetTypeActions/AssetTypeActions_ProceduralFoliageSpawner.h
@@ -10,4 +10,5 @@ public:
 	virtual FColor GetTypeColor() const override { return FColor(89, 134, 178); }
 	virtual UClass* GetSupportedClass() const override;
 	virtual uint32 GetCategories() override { return EAssetTypeCategories::Misc; }
+	virtual bool CanFilter() override;
 };
\ No newline at end of file
diff --git a/Engine/Source/Editor/FoliageEdit/Private/ActorFactoryProceduralFoliage.cpp b/Engine/Source/Editor/FoliageEdit/Private/ActorFactoryProceduralFoliage.cpp
--- a/Engine/Source/Editor/FoliageEdit/Private/ActorFactoryProceduralFoliage.cpp
+++ b/Engine/Source/Editor/FoliageEdit/Private/ActorFactoryProceduralFoliage.cpp
@@ -21,16 +21,21 @@ UActorFactoryProceduralFoliage
 UActorFactoryProceduralFoliage::UActorFactoryProceduralFoliage(const FObjectInitializer& ObjectInitializer)
 : Super(ObjectInitializer)
 {
-	DisplayName = LOCTEXT("ProceduralFoliageDisplayName", "ProceduralFoliage");
+	DisplayName = LOCTEXT("ProceduralFoliageDisplayName", "Procedural Foliage Volume");
 	NewActorClass = AProceduralFoliageVolume::StaticClass();
 	bUseSurfaceOrientation = true;
 }
 
+bool UActorFactoryProceduralFoliage::PreSpawnActor(UObject* Asset, FTransform& InOutLocation)
+{
+	return GetDefault<UEditorExperimentalSettings>()->bProceduralFoliage;
+}
+
 bool UActorFactoryProceduralFoliage::CanCreateActorFrom(const FAssetData& AssetData, FText& OutErrorMsg)
 {
 	if (!AssetData.IsValid() || !AssetData.GetClass()->IsChildOf(UProceduralFoliageSpawner::StaticClass()))
 	{
-		OutErrorMsg = NSLOCTEXT("CanCreateActor", "NoProceduralFoliage", "A valid ProceduralFoliage must be specified.");
+		OutErrorMsg = NSLOCTEXT("CanCreateActor", "NoProceduralFoliageSpawner", "A valid ProceduralFoliageSpawner must be specified.");
 		return false;
 	}
 
diff --git a/Engine/Source/Editor/FoliageEdit/Private/ActorFactoryProceduralFoliage.h b/Engine/Source/Editor/FoliageEdit/Private/ActorFactoryProceduralFoliage.h
--- a/Engine/Source/Editor/FoliageEdit/Private/ActorFactoryProceduralFoliage.h
+++ b/Engine/Source/Editor/FoliageEdit/Private/ActorFactoryProceduralFoliage.h
@@ -11,6 +11,7 @@ class UActorFactoryProceduralFoliage : public UActorFactoryBoxVolume
 	GENERATED_UCLASS_BODY()
 
 	// Begin UActorFactory Interface
+	virtual bool PreSpawnActor(UObject* Asset, FTransform& InOutLocation);
 	virtual bool CanCreateActorFrom( const FAssetData& AssetData, FText& OutErrorMsg ) override;
 	virtual void PostSpawnActor( UObject* Asset, AActor* NewActor) override;
 	virtual void PostCreateBlueprint( UObject* Asset, AActor* CDO ) override;
diff --git a/Engine/Source/Editor/FoliageEdit/Private/FoliageEditModule.cpp b/Engine/Source/Editor/FoliageEdit/Private/FoliageEditModule.cpp
--- a/Engine/Source/Editor/FoliageEdit/Private/FoliageEditModule.cpp
+++ b/Engine/Source/Editor/FoliageEdit/Private/FoliageEditModule.cpp
@@ -15,6 +15,7 @@ const FName FoliageEditAppIdentifier = FName(TEXT("FoliageEdApp"));
 #include "ActorFactoryProceduralFoliage.h"
 #include "ComponentVisualizer.h"
 #include "ProceduralFoliageVolume.h"
+#include "ProceduralFoliageBlockingVolume.h"
 #include "ProceduralFoliageComponent.h"
 #include "FoliageTypeObjectCustomization.h"
 
@@ -52,6 +53,13 @@ public:
 		GEditor->ActorFactories.Add(ProceduralFoliageVolumeFactory);
 
 #if WITH_EDITOR
+		// Volume placeability
+		if (!GetDefault<UEditorExperimentalSettings>()->bProceduralFoliage)
+		{
+			AProceduralFoliageVolume::StaticClass()->ClassFlags |= CLASS_NotPlaceable;
+			AProceduralFoliageBlockingVolume::StaticClass()->ClassFlags |= CLASS_NotPlaceable;
+		}
+
 		SubscribeEvents();
 #endif
 	}
@@ -98,14 +106,34 @@ public:
 	{
 		GEngine->OnLevelActorDeleted().Remove(OnLevelActorDeletedDelegateHandle);
 		OnLevelActorDeletedDelegateHandle = GEngine->OnLevelActorDeleted().AddRaw(this, &FFoliageEditModule::OnLevelActorDeleted);
+
+		auto ExperimentalSettings = GetMutableDefault<UEditorExperimentalSettings>();
+		ExperimentalSettings->OnSettingChanged().Remove(OnExperimentalSettingChangedDelegateHandle);
+		OnExperimentalSettingChangedDelegateHandle =  ExperimentalSettings->OnSettingChanged().AddRaw(this, &FFoliageEditModule::HandleExperimentalSettingChanged);
 	}
 
 	void UnsubscribeEvents()
 	{
 		GEngine->OnLevelActorDeleted().Remove(OnLevelActorDeletedDelegateHandle);
+		GetMutableDefault<UEditorExperimentalSettings>()->OnSettingChanged().Remove(OnExperimentalSettingChangedDelegateHandle);
+	}
+
+	void HandleExperimentalSettingChanged(FName PropertyName)
+	{
+		if (GetDefault<UEditorExperimentalSettings>()->bProceduralFoliage)
+		{
+			AProceduralFoliageVolume::StaticClass()->ClassFlags &= ~CLASS_NotPlaceable;
+			AProceduralFoliageBlockingVolume::StaticClass()->ClassFlags &= ~CLASS_NotPlaceable;
+		}
+		else
+		{
+			AProceduralFoliageVolume::StaticClass()->ClassFlags |= CLASS_NotPlaceable;
+			AProceduralFoliageBlockingVolume::StaticClass()->ClassFlags |= CLASS_NotPlaceable;
+		}
 	}
 
 	FDelegateHandle OnLevelActorDeletedDelegateHandle;
+	FDelegateHandle OnExperimentalSettingChangedDelegateHandle;
 #endif
 };
 
diff --git a/Engine/Source/Editor/FoliageEdit/Private/FoliageTypeDetails.cpp b/Engine/Source/Editor/FoliageEdit/Private/FoliageTypeDetails.cpp
--- a/Engine/Source/Editor/FoliageEdit/Private/FoliageTypeDetails.cpp
+++ b/Engine/Source/Editor/FoliageEdit/Private/FoliageTypeDetails.cpp
@@ -98,7 +98,17 @@ void FFoliageTypeDetails::CustomizeDetails(IDetailLayoutBuilder& DetailBuilder)
 
 	TMap<const FName, IDetailPropertyRow*> PropertyRowsByName;
 	AddSubcategoryProperties(DetailBuilder, PlacementName, &PropertyRowsByName);
-	AddSubcategoryProperties(DetailBuilder, ProceduralName);
+
+	// If enabled, show the properties for procedural placement
+	if (GetDefault<UEditorExperimentalSettings>()->bProceduralFoliage)
+	{
+		AddSubcategoryProperties(DetailBuilder, ProceduralName);
+	}
+	else
+	{
+		FFoliageTypeCustomizationHelpers::HideFoliageCategory(DetailBuilder, ProceduralName);
+	}
+	
 	AddSubcategoryProperties(DetailBuilder, InstanceSettingsName);
 
 	// Custom visibility overrides for the scale axes
diff --git a/Engine/Source/Editor/FoliageEdit/Private/ProceduralFoliageSpawnerFactory.cpp b/Engine/Source/Editor/FoliageEdit/Private/ProceduralFoliageSpawnerFactory.cpp
--- a/Engine/Source/Editor/FoliageEdit/Private/ProceduralFoliageSpawnerFactory.cpp
+++ b/Engine/Source/Editor/FoliageEdit/Private/ProceduralFoliageSpawnerFactory.cpp
@@ -23,4 +23,9 @@ UObject* UProceduralFoliageSpawnerFactory::FactoryCreateNew(UClass* Class, UObje
 uint32 UProceduralFoliageSpawnerFactory::GetMenuCategories() const
 {
 	return EAssetTypeCategories::Misc;
+}
+
+bool UProceduralFoliageSpawnerFactory::ShouldShowInNewMenu() const
+{
+	return GetDefault<UEditorExperimentalSettings>()->bProceduralFoliage;
 }
\ No newline at end of file
diff --git a/Engine/Source/Editor/FoliageEdit/Private/ProceduralFoliageSpawnerFactory.h b/Engine/Source/Editor/FoliageEdit/Private/ProceduralFoliageSpawnerFactory.h
--- a/Engine/Source/Editor/FoliageEdit/Private/ProceduralFoliageSpawnerFactory.h
+++ b/Engine/Source/Editor/FoliageEdit/Private/ProceduralFoliageSpawnerFactory.h
@@ -17,5 +17,6 @@ class UProceduralFoliageSpawnerFactory : public UFactory
 	// UFactory interface
 	virtual UObject* FactoryCreateNew(UClass* Class, UObject* InParent, FName Name, EObjectFlags Flags, UObject* Context, FFeedbackContext* Warn) override;
 	virtual uint32 GetMenuCategories() const override;
+	virtual bool ShouldShowInNewMenu() const override;
 	// End of UFactory interface
 };
diff --git a/Engine/Source/Editor/UnrealEd/Classes/Settings/EditorExperimentalSettings.h b/Engine/Source/Editor/UnrealEd/Classes/Settings/EditorExperimentalSettings.h
--- a/Engine/Source/Editor/UnrealEd/Classes/Settings/EditorExperimentalSettings.h
+++ b/Engine/Source/Editor/UnrealEd/Classes/Settings/EditorExperimentalSettings.h
@@ -21,6 +21,10 @@ class UNREALED_API UEditorExperimentalSettings
 
 public:
 
+	/** Allows usage of the procedural foliage system */
+	UPROPERTY(EditAnywhere, config, Category = Foliage, meta = (DisplayName = "Procedural Foliage"))
+	bool bProceduralFoliage;
+
 	/** Allows usage of the Localization Dashboard */
 	UPROPERTY(EditAnywhere, config, Category = Tools, meta = (DisplayName = "Localization Dashboard"))
 	bool bEnableLocalizationDashboard;
@@ -88,7 +92,7 @@ public:
 	/** Enables Environment Queries editor */
 	UPROPERTY(EditAnywhere, config, Category = AI, meta = (DisplayName = "Environment Querying System"))
 	bool bEQSEditor;
-	
+
 	/**
 	 * Returns an event delegate that is executed when a setting has changed.
 	 *
diff --git a/Engine/Source/Runtime/Foliage/Private/ProceduralFoliageBlockingVolume.h b/Engine/Source/Runtime/Foliage/Private/ProceduralFoliageBlockingVolume.h
deleted file mode 100644
diff --git a/Engine/Source/Runtime/Foliage/Public/ProceduralFoliageBlockingVolume.h b/Engine/Source/Runtime/Foliage/Public/ProceduralFoliageBlockingVolume.h
new file mode 100644
--- /dev/null
+++ b/Engine/Source/Runtime/Foliage/Public/ProceduralFoliageBlockingVolume.h
@@ -0,0 +1,21 @@
+// Copyright 1998-2015 Epic Games, Inc. All Rights Reserved.
+
+
+#pragma once
+#include "GameFramework/Volume.h"
+#include "ProceduralFoliageBlockingVolume.generated.h"
+
+class AProceduralFoliageVolume;
+
+/** An invisible volume used to block ProceduralFoliage instances from being spawned. */
+UCLASS(MinimalAPI)
+class AProceduralFoliageBlockingVolume : public AVolume
+{
+	GENERATED_UCLASS_BODY()
+
+	UPROPERTY(Category = ProceduralFoliage, EditAnywhere)
+	AProceduralFoliageVolume* ProceduralFoliageVolume;
+};
+
+
+
