Do not create cross-level references on using "snap to ground" foliage command

diff --git a/Engine/Source/Editor/FoliageEdit/Private/FoliageEdMode.cpp b/Engine/Source/Editor/FoliageEdit/Private/FoliageEdMode.cpp
--- a/Engine/Source/Editor/FoliageEdit/Private/FoliageEdMode.cpp
+++ b/Engine/Source/Editor/FoliageEdit/Private/FoliageEdMode.cpp
@@ -2025,10 +2025,16 @@ bool FEdModeFoliage::SnapInstanceToGround(AInstancedFoliageActor* InIFA, float A
 	{
 		UPrimitiveComponent* HitComponent = Hit.Component.Get();
 
+		if (HitComponent->GetComponentLevel() != InIFA->GetLevel())
+		{
+			// We should not create cross-level references automatically
+			return false;
+		}
+
 		// We cannot be based on an a blueprint component as these will disappear when the construction script is re-run
 		if (HitComponent->IsCreatedByConstructionScript())
 		{
-			HitComponent = nullptr;
+			return false;
 		}
 																
 		// Find BSP brush 
