Added: Foliage command to find and select 'invalid' instances

diff --git a/Engine/Source/Editor/FoliageEdit/Private/FoliageEdMode.cpp b/Engine/Source/Editor/FoliageEdit/Private/FoliageEdMode.cpp
--- a/Engine/Source/Editor/FoliageEdit/Private/FoliageEdMode.cpp
+++ b/Engine/Source/Editor/FoliageEdit/Private/FoliageEdMode.cpp
@@ -1194,6 +1194,90 @@ void FEdModeFoliage::RemoveSelectedInstances(UWorld* InWorld)
 	GEditor->EndTransaction();
 }
 
+TAutoConsoleVariable<float> CVarOffGroundTreshold(
+	TEXT("foliage.OffGroundThreshold"),
+	5.0f,
+	TEXT("Maximum distance from base component (in local space) at which instance is still considered as valid"));
+
+void FEdModeFoliage::SelectInvalidInstances(const UFoliageType* Settings)
+{
+	UWorld* InWorld = GetWorld();
+	
+	static FName NAME_FoliageGroundCheck = FName("FoliageGroundCheck");
+	FCollisionQueryParams QueryParams(NAME_FoliageGroundCheck, true);
+	QueryParams.bReturnFaceIndex = false;
+	FCollisionShape SphereShape;
+	SphereShape.SetSphere(0.f);
+	float InstanceOffGroundLocalThreshold = CVarOffGroundTreshold.GetValueOnGameThread();
+		
+	for (FFoliageMeshInfoIterator It(InWorld, Settings); It; ++It)
+	{
+		FFoliageMeshInfo* MeshInfo = (*It);
+		AInstancedFoliageActor* IFA = It.GetActor();
+		int32 NumInstances = MeshInfo->Instances.Num();
+		TArray<FHitResult> Hits; Hits.Reserve(16);
+
+		TArray<int32> InvalidInstances;
+		
+		for (int32 InstanceIdx = 0; InstanceIdx < NumInstances; ++InstanceIdx)
+		{
+			FFoliageInstance& Instance = MeshInfo->Instances[InstanceIdx];
+			UActorComponent* CurrentInstanceBase = IFA->InstanceBaseCache.GetInstanceBasePtr(Instance.BaseId).Get();
+			bool bInvalidInstance = true;
+
+			if (CurrentInstanceBase != nullptr)
+			{
+				FVector InstanceTraceRange = Instance.GetInstanceWorldTransform().TransformVector(FVector(0.f, 0.f, 1000.f));
+				FVector Start = Instance.Location + InstanceTraceRange;
+				FVector End = Instance.Location - InstanceTraceRange;
+				
+				InWorld->SweepMultiByObjectType(Hits, Start, End, FQuat::Identity, FCollisionObjectQueryParams(ECC_WorldStatic), SphereShape, QueryParams);
+				
+				for (const FHitResult& Hit : Hits)
+				{
+					UPrimitiveComponent* HitComponent = Hit.GetComponent();
+					if (HitComponent->IsCreatedByConstructionScript())
+					{
+						continue;
+					}
+																
+					UModelComponent* ModelComponent = Cast<UModelComponent>(HitComponent);
+					if (ModelComponent)
+					{
+						ABrush* BrushActor = ModelComponent->GetModel()->FindBrush(Hit.Location);
+						if (BrushActor)
+						{
+							HitComponent = BrushActor->GetBrushComponent();
+						}
+					}
+					
+					if (HitComponent == CurrentInstanceBase)
+					{
+						FVector InstanceWorldZOffset = Instance.GetInstanceWorldTransform().TransformVector(FVector(0.f, 0.f, Instance.ZOffset));
+						float DistanceToGround = FVector::Dist(Instance.Location, Hit.Location + InstanceWorldZOffset);
+						float InstanceWorldTreshold = Instance.GetInstanceWorldTransform().TransformVector(FVector(0.f, 0.f, InstanceOffGroundLocalThreshold)).Size();
+																								
+						if ((DistanceToGround - InstanceWorldTreshold) <= KINDA_SMALL_NUMBER)
+						{
+							bInvalidInstance = false;
+						}
+					}
+				}
+			}
+			
+			if (bInvalidInstance)
+			{
+				InvalidInstances.Add(InstanceIdx);
+			}
+		}
+
+		if (InvalidInstances.Num() > 0)
+		{
+			MeshInfo->SelectInstances(IFA, true, InvalidInstances);
+		}
+	}
+}
+
 void FEdModeFoliage::ReapplyInstancesForBrush(UWorld* InWorld, const UFoliageType* Settings, const FSphere& BrushSphere, float Pressure)
 {
 	// Adjust instance density first
diff --git a/Engine/Source/Editor/FoliageEdit/Private/FoliageEdMode.h b/Engine/Source/Editor/FoliageEdit/Private/FoliageEdMode.h
--- a/Engine/Source/Editor/FoliageEdit/Private/FoliageEdMode.h
+++ b/Engine/Source/Editor/FoliageEdit/Private/FoliageEdMode.h
@@ -329,6 +329,9 @@ public:
 	/** Set/Clear selection for foliage instances of specific type  */
 	void SelectInstances(const UFoliageType* Settings, bool bSelect);
 
+	/** Find and select instances that don't have valid base or 'off-ground' */
+	void SelectInvalidInstances(const UFoliageType* Settings);
+
 	/** Add desired instances. Uses foliage settings to determine location/scale/rotation and whether instances should be ignored */
 	static void AddInstances(UWorld* InWorld, const TArray<FDesiredFoliageInstance>& DesiredInstances);
 
diff --git a/Engine/Source/Editor/FoliageEdit/Private/FoliageEditActions.cpp b/Engine/Source/Editor/FoliageEdit/Private/FoliageEditActions.cpp
--- a/Engine/Source/Editor/FoliageEdit/Private/FoliageEditActions.cpp
+++ b/Engine/Source/Editor/FoliageEdit/Private/FoliageEditActions.cpp
@@ -24,6 +24,7 @@ void FFoliageEditCommands::RegisterCommands()
 	UI_COMMAND( ShowFoliageTypeInCB, "Show in Content Browser", "Show asset in Content Browser.", EUserInterfaceActionType::Button, FInputChord() );
 	UI_COMMAND( SelectAllInstances, "Select All Instances", "Select all instances of foliage type.", EUserInterfaceActionType::Button, FInputChord());
 	UI_COMMAND( DeselectAllInstances, "Deselect All Instances", "Deselect all instances of foliage type.", EUserInterfaceActionType::Button, FInputChord());
+	UI_COMMAND( SelectInvalidInstances, "Select Invalid Instances", "Select all instances that are off ground.", EUserInterfaceActionType::Button, FInputChord());
 }
 
 #undef LOCTEXT_NAMESPACE
\ No newline at end of file
diff --git a/Engine/Source/Editor/FoliageEdit/Private/FoliageEditActions.h b/Engine/Source/Editor/FoliageEdit/Private/FoliageEditActions.h
--- a/Engine/Source/Editor/FoliageEdit/Private/FoliageEditActions.h
+++ b/Engine/Source/Editor/FoliageEdit/Private/FoliageEditActions.h
@@ -42,6 +42,7 @@ public:
 	TSharedPtr< FUICommandInfo > ShowFoliageTypeInCB;
 	TSharedPtr< FUICommandInfo > SelectAllInstances;
 	TSharedPtr< FUICommandInfo > DeselectAllInstances;
+	TSharedPtr< FUICommandInfo > SelectInvalidInstances;
 	/**
 	 * Initialize commands
 	 */
diff --git a/Engine/Source/Editor/FoliageEdit/Private/SFoliageEdit.cpp b/Engine/Source/Editor/FoliageEdit/Private/SFoliageEdit.cpp
--- a/Engine/Source/Editor/FoliageEdit/Private/SFoliageEdit.cpp
+++ b/Engine/Source/Editor/FoliageEdit/Private/SFoliageEdit.cpp
@@ -555,6 +555,12 @@ void SFoliageEdit::BindCommands()
 		FExecuteAction::CreateSP(this, &SFoliageEdit::OnDeselectAllInstances),
 		FCanExecuteAction(),
 		FIsActionChecked());
+
+	UICommandList->MapAction(
+		Commands.SelectInvalidInstances,
+		FExecuteAction::CreateSP(this, &SFoliageEdit::OnSelectInvalidInstances),
+		FCanExecuteAction(),
+		FIsActionChecked());
 }
 
 TSharedRef<SWidget> SFoliageEdit::BuildToolBar()
@@ -790,21 +796,25 @@ FReply SFoliageEdit::OnDrop_ListView(const FGeometry& MyGeometry, const FDragDro
 TSharedPtr<SWidget> SFoliageEdit::ConstructFoliageMeshContextMenu() const
 {
 	const FFoliageEditCommands& Commands = FFoliageEditCommands::Get();
-	
 	FMenuBuilder MenuBuilder(true, UICommandList);
-	MenuBuilder.AddMenuEntry(Commands.RemoveFoliageType);
 
-	MenuBuilder.AddSubMenu(
-		LOCTEXT( "ReplaceFoliageType", "Replace Foliage Type" ),
-		LOCTEXT ("ReplaceFoliageType_ToolTip", "Replaces selected foliage type with another foliage type asset"),
-		FNewMenuDelegate::CreateSP(this, &SFoliageEdit::FillReplaceFoliageTypeSubmenu));
+	if (MeshTreeWidget->GetSelectedItems().Num() > 0)
+	{
+		MenuBuilder.AddMenuEntry(Commands.RemoveFoliageType);
 
-	MenuBuilder.AddMenuEntry(Commands.ShowFoliageTypeInCB);
+		MenuBuilder.AddSubMenu(
+			LOCTEXT( "ReplaceFoliageType", "Replace Foliage Type" ),
+			LOCTEXT ("ReplaceFoliageType_ToolTip", "Replaces selected foliage type with another foliage type asset"),
+			FNewMenuDelegate::CreateSP(this, &SFoliageEdit::FillReplaceFoliageTypeSubmenu));
 
-	if (IsSelectTool() || IsLassoSelectTool())
-	{
-		MenuBuilder.AddMenuEntry(Commands.SelectAllInstances);
-		MenuBuilder.AddMenuEntry(Commands.DeselectAllInstances);
+		MenuBuilder.AddMenuEntry(Commands.ShowFoliageTypeInCB);
+
+		if (IsSelectTool() || IsLassoSelectTool())
+		{
+			MenuBuilder.AddMenuEntry(Commands.SelectAllInstances);
+			MenuBuilder.AddMenuEntry(Commands.DeselectAllInstances);
+			MenuBuilder.AddMenuEntry(Commands.SelectInvalidInstances);
+		}
 	}
 
 	return MenuBuilder.MakeWidget();
@@ -1132,6 +1142,16 @@ void SFoliageEdit::OnDeselectAllInstances()
 	}
 }
 
+void SFoliageEdit::OnSelectInvalidInstances()
+{
+	auto MeshUIList = MeshTreeWidget->GetSelectedItems();
+	for (FFoliageMeshUIInfoPtr MeshUIPtr : MeshUIList)
+	{
+		const UFoliageType* FoliageType = MeshUIPtr->Settings;
+		FoliageEditMode->SelectInvalidInstances(FoliageType);
+	}
+}
+
 ECheckBoxState SFoliageEdit::GetState_AllMeshes() const
 {
 	bool bHasChecked = false;
diff --git a/Engine/Source/Editor/FoliageEdit/Private/SFoliageEdit.h b/Engine/Source/Editor/FoliageEdit/Private/SFoliageEdit.h
--- a/Engine/Source/Editor/FoliageEdit/Private/SFoliageEdit.h
+++ b/Engine/Source/Editor/FoliageEdit/Private/SFoliageEdit.h
@@ -189,6 +189,9 @@ private:
 
 	/** Handler for 'Deselect All' command  */
 	void OnDeselectAllInstances();
+
+	/** Handler for 'Select Invalid Instances' command  */
+	void OnSelectInvalidInstances();
 	
 	/** Toggle all meshes on/off */
 	ECheckBoxState GetState_AllMeshes() const;
