[UE-13896] - Drag-drop of foliage types is treated as a single transaction again

diff --git a/Engine/Source/Editor/FoliageEdit/Private/SFoliageEdit.cpp b/Engine/Source/Editor/FoliageEdit/Private/SFoliageEdit.cpp
--- a/Engine/Source/Editor/FoliageEdit/Private/SFoliageEdit.cpp
+++ b/Engine/Source/Editor/FoliageEdit/Private/SFoliageEdit.cpp
@@ -785,6 +785,9 @@ FReply SFoliageEdit::OnDrop_ListView(const FGeometry& MyGeometry, const FDragDro
 {
 	TArray<FAssetData> DroppedAssetData = AssetUtil::ExtractAssetDataFromDrag(DragDropEvent);
 	{
+		// Treat the entire drop as a transaction (in case multiples types are being added)
+		const FScopedTransaction Transaction(NSLOCTEXT("UnrealEd", "FoliageMode_DragDropTypesTransaction", "Drag-drop Foliage Type(s)"));
+
 		for (int32 AssetIdx = 0; AssetIdx < DroppedAssetData.Num(); ++AssetIdx)
 		{	
 			AddFoliageType(DroppedAssetData[AssetIdx]);
