Added tooltip to foliage tool that shows total instance count in current level and across world

diff --git a/Engine/Source/Editor/FoliageEdit/Private/FoliageEdMode.cpp b/Engine/Source/Editor/FoliageEdit/Private/FoliageEdMode.cpp
--- a/Engine/Source/Editor/FoliageEdit/Private/FoliageEdMode.cpp
+++ b/Engine/Source/Editor/FoliageEdit/Private/FoliageEdMode.cpp
@@ -2258,6 +2258,36 @@ void FEdModeFoliage::OnInstanceCountUpdated(const UFoliageType* FoliageType)
 	FoliageMeshList[EntryIndex]->InstanceCountCurrentLevel = InstanceCountCurrentLevel;
 }
 
+void FEdModeFoliage::CalcTotalInstanceCount(int32& OutInstanceCountTotal, int32& OutInstanceCountCurrentLevel)
+{
+	OutInstanceCountTotal = 0;
+	OutInstanceCountCurrentLevel = 0;
+	UWorld* InWorld = GetWorld();
+	ULevel* CurrentLevel = InWorld->GetCurrentLevel(); 
+
+	const int32 NumLevels = InWorld->GetNumLevels();
+	for (int32 LevelIdx = 0; LevelIdx < NumLevels; ++LevelIdx)
+	{
+		ULevel* Level = InWorld->GetLevel(LevelIdx);
+		AInstancedFoliageActor* IFA = AInstancedFoliageActor::GetInstancedFoliageActorForLevel(Level);
+		if (IFA)
+		{
+			int32 IFAInstanceCount = 0;
+			for (const auto& MeshPair : IFA->FoliageMeshes)
+			{
+				const FFoliageMeshInfo& MeshInfo = *MeshPair.Value;
+				IFAInstanceCount+= MeshInfo.Instances.Num();
+			}
+
+			OutInstanceCountTotal+= IFAInstanceCount;
+			if (CurrentLevel == Level)
+			{
+				OutInstanceCountCurrentLevel+= IFAInstanceCount;
+			}
+		}
+	}
+}
+
 bool FEdModeFoliage::CanPaint(const ULevel* InLevel)
 {
 	for (const auto& MeshUIPtr : FoliageMeshList)
diff --git a/Engine/Source/Editor/FoliageEdit/Private/FoliageEdMode.h b/Engine/Source/Editor/FoliageEdit/Private/FoliageEdMode.h
--- a/Engine/Source/Editor/FoliageEdit/Private/FoliageEdMode.h
+++ b/Engine/Source/Editor/FoliageEdit/Private/FoliageEdMode.h
@@ -293,6 +293,9 @@ public:
 	/** Handler for foliage mesh instance count changes */
 	void OnInstanceCountUpdated(const UFoliageType* FoliageType);
 
+	/** Counts total number of instances in current level and across whole world */
+	void CalcTotalInstanceCount(int32& OutInstanceCountTotal, int32& OutInstanceCountCurrentLevel);
+
 	/** Whether any of the selected foliage types can be painted into level */
 	bool CanPaint(const ULevel* InLevel);
 		
diff --git a/Engine/Source/Editor/FoliageEdit/Private/SFoliageEdit.cpp b/Engine/Source/Editor/FoliageEdit/Private/SFoliageEdit.cpp
--- a/Engine/Source/Editor/FoliageEdit/Private/SFoliageEdit.cpp
+++ b/Engine/Source/Editor/FoliageEdit/Private/SFoliageEdit.cpp
@@ -444,6 +444,7 @@ void SFoliageEdit::Construct(const FArguments& InArgs)
 						+ SHeaderRow::Column(FoliageMeshColumns::ColumnID_InstanceCount)
 						.HeaderContentPadding(FMargin(10, 1, 0, 1))
 						.DefaultLabel(LOCTEXT("InstanceCount", "Count"))
+						.DefaultTooltip(this, &SFoliageEdit::GetTotalInstanceCountTooltipText)
 						.FixedWidth(60.f)
 
 						// save Foliage asset
@@ -1206,6 +1207,17 @@ void SFoliageEdit::OnMeshesColumnSortModeChanged(EColumnSortPriority::Type InPri
 	FoliageEditMode->OnFoliageMeshListSortModeChanged(InSortMode);
 }
 
+FText SFoliageEdit::GetTotalInstanceCountTooltipText() const
+{
+	// Probably should cache these values, 
+	// but we call this only occasionally when tooltip is active
+	int32 InstanceCountTotal = 0;
+	int32 InstanceCountCurrentLevel = 0;
+	FoliageEditMode->CalcTotalInstanceCount(InstanceCountTotal, InstanceCountCurrentLevel);
+			
+	return FText::Format(LOCTEXT("FoliageTotalInstanceCount", "Current Level: {0} Total: {1}"), FText::AsNumber(InstanceCountCurrentLevel), FText::AsNumber(InstanceCountTotal));
+}
+
 void SFoliageEdit::HandleOnToolChanged()
 {
 	RefreshMeshDetailsWidget();
diff --git a/Engine/Source/Editor/FoliageEdit/Private/SFoliageEdit.h b/Engine/Source/Editor/FoliageEdit/Private/SFoliageEdit.h
--- a/Engine/Source/Editor/FoliageEdit/Private/SFoliageEdit.h
+++ b/Engine/Source/Editor/FoliageEdit/Private/SFoliageEdit.h
@@ -204,6 +204,9 @@ private:
 	EColumnSortMode::Type GetMeshColumnSortMode() const;
 	void OnMeshesColumnSortModeChanged(EColumnSortPriority::Type InPriority, const FName& InColumnName, EColumnSortMode::Type InSortMode);
 
+	/** Tooltip text for 'Instance Count" column */
+	FText GetTotalInstanceCountTooltipText() const;
+
 	/** Handler to trigger a refresh of the details view when the active tool changes */
 	void HandleOnToolChanged();
 
