Usability improvements to foliage edit mode UI

- Tools multibox laid out vertically to keep spacing constant and match the style of the Placement Mode UI
- Formatting matches details panel
- Added thumbnail view in addition to the tree
- Can activate/deactivate types via double-click
- Can collapse the details area
- Selection modes have buttons for selection on ALL instances
- Tooltips for filters are mode-specific
- Activation toggle applies to all selected types when checked/unchecked

diff --git a/Engine/Source/Editor/EditorStyle/Private/SlateEditorStyle.cpp b/Engine/Source/Editor/EditorStyle/Private/SlateEditorStyle.cpp
--- a/Engine/Source/Editor/EditorStyle/Private/SlateEditorStyle.cpp
+++ b/Engine/Source/Editor/EditorStyle/Private/SlateEditorStyle.cpp
@@ -1587,12 +1587,52 @@ void FSlateEditorStyle::FStyle::SetupGeneralStyles()
 	}
 
 	// Foliage Edit Mode
-		{
-		Set( "FoliageEditMode.SetPaint", new IMAGE_BRUSH( "Icons/FoliageEditMode/icon_FoliageEdMode_Paint_40x", Icon40x40 ) );
-		Set( "FoliageEditMode.SetReapplySettings", new IMAGE_BRUSH( "Icons/FoliageEditMode/icon_FoliageEdMode_Reapply_40x", Icon40x40 ) );
-		Set( "FoliageEditMode.SetSelect", new IMAGE_BRUSH( "Icons/FoliageEditMode/icon_FoliageEdMode_Select_40x", Icon40x40 ) );
-		Set( "FoliageEditMode.SetLassoSelect", new IMAGE_BRUSH( "Icons/FoliageEditMode/icon_FoliageEdMode_Lasso_40x", Icon40x40 ) );
-		Set( "FoliageEditMode.SetPaintBucket", new IMAGE_BRUSH( "Icons/FoliageEditMode/icon_FoliageEdMode_PaintBucket_40x", Icon40x40 ) );
+	{	
+		FLinearColor DimBackground = FLinearColor(FColor(64, 64, 64));
+		FLinearColor DimBackgroundHover = FLinearColor(FColor(50, 50, 50));
+		FLinearColor DarkBackground = FLinearColor(FColor(42, 42, 42));
+
+		Set("FoliageEditToolBar.ToggleButton", FCheckBoxStyle()
+			.SetCheckBoxType(ESlateCheckBoxType::ToggleButton)
+			.SetUncheckedImage(BOX_BRUSH("Common/Selection", 8.0f / 32.0f, DimBackground))
+			.SetUncheckedPressedImage(BOX_BRUSH("PlacementMode/TabActive", 8.0f / 32.0f))
+			.SetUncheckedHoveredImage(BOX_BRUSH("Common/Selection", 8.0f / 32.0f, DimBackgroundHover))
+			.SetCheckedImage(BOX_BRUSH("PlacementMode/TabActive", 8.0f / 32.0f))
+			.SetCheckedHoveredImage(BOX_BRUSH("PlacementMode/TabActive", 8.0f / 32.0f))
+			.SetCheckedPressedImage(BOX_BRUSH("PlacementMode/TabActive", 8.0f / 32.0f))
+			.SetPadding(0));
+
+		Set("FoliageEditToolBar.Background", new BOX_BRUSH("Common/GroupBorder", FMargin(4.0f / 16.0f)));
+		Set("FoliageEditToolBar.Icon", new IMAGE_BRUSH("Icons/icon_tab_Toolbars_16x", Icon16x16));
+		Set("FoliageEditToolBar.Expand", new IMAGE_BRUSH("Icons/toolbar_expand_16x", Icon16x16));
+		Set("FoliageEditToolBar.SubMenuIndicator", new IMAGE_BRUSH("Common/SubmenuArrow", Icon8x8));
+		Set("FoliageEditToolBar.SToolBarComboButtonBlock.Padding", FMargin(4.0f));
+		Set("FoliageEditToolBar.SToolBarButtonBlock.Padding", FMargin(0.f));
+		Set("FoliageEditToolBar.SToolBarCheckComboButtonBlock.Padding", FMargin(4.0f));
+		Set("FoliageEditToolBar.SToolBarButtonBlock.CheckBox.Padding", FMargin(10.0f, 6.f));
+		Set("FoliageEditToolBar.SToolBarComboButtonBlock.ComboButton.Color", DefaultForeground);
+
+		Set("FoliageEditToolBar.Block.IndentedPadding", FMargin(18.0f, 2.0f, 4.0f, 4.0f));
+		Set("FoliageEditToolBar.Block.Padding", FMargin(2.0f, 2.0f, 4.0f, 4.0f));
+
+		Set("FoliageEditToolBar.Separator", new BOX_BRUSH("Old/Button", 4.0f / 32.0f));
+		Set("FoliageEditToolBar.Separator.Padding", FMargin(0.5f));
+
+		Set("FoliageEditToolBar.Label", FTextBlockStyle(NormalText).SetFont(TTF_CORE_FONT("Fonts/Roboto-Regular", 9)));
+		Set("FoliageEditToolBar.EditableText", FEditableTextBoxStyle(NormalEditableTextBoxStyle).SetFont(TTF_CORE_FONT("Fonts/Roboto-Regular", 9)));
+		Set("FoliageEditToolBar.Keybinding", FTextBlockStyle(NormalText).SetFont(TTF_CORE_FONT("Fonts/Roboto-Regular", 8)));
+
+		Set("FoliageEditToolBar.Heading", FTextBlockStyle(NormalText)
+			.SetFont(TTF_CORE_FONT("Fonts/Roboto-Regular", 8))
+			.SetColorAndOpacity(FLinearColor(0.4f, 0.4, 0.4f, 1.0f)));
+
+		// Set the large icons to be the same as the small ones
+		Set("FoliageEditMode.SetPaint", new IMAGE_BRUSH("Icons/FoliageEditMode/icon_FoliageEdMode_Paint_40x", Icon20x20));
+		Set("FoliageEditMode.SetReapplySettings", new IMAGE_BRUSH("Icons/FoliageEditMode/icon_FoliageEdMode_Reapply_40x", Icon20x20));
+		Set("FoliageEditMode.SetSelect", new IMAGE_BRUSH("Icons/FoliageEditMode/icon_FoliageEdMode_Select_40x", Icon20x20));
+		Set("FoliageEditMode.SetLassoSelect", new IMAGE_BRUSH("Icons/FoliageEditMode/icon_FoliageEdMode_Lasso_40x", Icon20x20));
+		Set("FoliageEditMode.SetPaintBucket", new IMAGE_BRUSH("Icons/FoliageEditMode/icon_FoliageEdMode_PaintBucket_40x", Icon20x20));
+
 		Set( "FoliageEditMode.SetPaint.Small", new IMAGE_BRUSH( "Icons/FoliageEditMode/icon_FoliageEdMode_Paint_40x", Icon20x20 ) );
 		Set( "FoliageEditMode.SetReapplySettings.Small", new IMAGE_BRUSH( "Icons/FoliageEditMode/icon_FoliageEdMode_Reapply_40x", Icon20x20 ) );
 		Set( "FoliageEditMode.SetSelect.Small", new IMAGE_BRUSH( "Icons/FoliageEditMode/icon_FoliageEdMode_Select_40x", Icon20x20 ) );
@@ -1613,6 +1653,21 @@ void FSlateEditorStyle::FStyle::SetupGeneralStyles()
 		Set( "FoliageEditMode.ItemBackground", new IMAGE_BRUSH( "Icons/FoliageEditMode/FoliageEditMode_Background", Icon64x64 ) );
 		Set( "FoliageEditMode.BubbleBorder", new BOX_BRUSH( "Icons/FoliageEditMode/FoliageEditMode_BubbleBorder", FMargin(8/32.0f) ) );
 
+		Set( "FoliageEditMode.TreeView.ScrollBorder", FScrollBorderStyle()
+			.SetTopShadowBrush(FSlateNoResource())
+			.SetBottomShadowBrush(BOX_BRUSH("Common/ScrollBorderShadowBottom", FVector2D(16, 8), FMargin(0.5, 0, 0.5, 1)))
+			);
+
+		Set("FoliageEditMode.Splitter", FSplitterStyle()
+			.SetHandleNormalBrush(IMAGE_BRUSH("Common/SplitterHandleHighlight", Icon8x8, FLinearColor(.2f, .2f, .2f, 1.f)))
+			.SetHandleHighlightBrush(IMAGE_BRUSH("Common/SplitterHandleHighlight", Icon8x8, FLinearColor::White))
+			);
+
+		Set("FoliageEditMode.ActiveToolName.Text", FTextBlockStyle(NormalText)
+			.SetFont(TTF_CORE_FONT("Fonts/Roboto-Bold", 11))
+			.SetShadowOffset(FVector2D(1, 1))
+			);
+
 		Set("FoliageEditMode.AddFoliageType.Text", FTextBlockStyle(NormalText)
 			.SetFont(TTF_CORE_FONT("Fonts/Roboto-Bold", 10))
 			.SetColorAndOpacity(FLinearColor(1.0f, 1.0f, 1.0f))
diff --git a/Engine/Source/Editor/FoliageEdit/Private/FoliageEdMode.cpp b/Engine/Source/Editor/FoliageEdit/Private/FoliageEdMode.cpp
--- a/Engine/Source/Editor/FoliageEdit/Private/FoliageEdMode.cpp
+++ b/Engine/Source/Editor/FoliageEdit/Private/FoliageEdMode.cpp
@@ -2481,57 +2481,57 @@ void FEdModeFoliage::ReplaceSettingsObject(UFoliageType* OldSettings, UFoliageTy
 	PopulateFoliageMeshList();
 }
 
-UFoliageType* FEdModeFoliage::SaveSettingsObject(UFoliageType* Settings)
+UFoliageType* FEdModeFoliage::SaveFoliageTypeObject(UFoliageType* InFoliageType)
 {
-	FScopedTransaction Transaction(NSLOCTEXT("UnrealEd", "FoliageMode_SaveSettingsObject", "Foliage Editing: Save Settings Object"));
+	FScopedTransaction Transaction(NSLOCTEXT("UnrealEd", "FoliageMode_SaveFoliageTypeObject", "Foliage Editing: Save Foliage Type Object"));
 
-	UFoliageType* SettingsToSave = nullptr;
+	UFoliageType* TypeToSave = nullptr;
 
-	if (!Settings->IsAsset())
+	if (!InFoliageType->IsAsset())
 	{
 		FString PackageName;
-		UStaticMesh* StaticMesh = Settings->GetStaticMesh();
+		UStaticMesh* StaticMesh = InFoliageType->GetStaticMesh();
 		if (StaticMesh)
 		{
 			// Build default settings asset name and path
-			PackageName = FPackageName::GetLongPackagePath(StaticMesh->GetOutermost()->GetName()) + TEXT("/") + StaticMesh->GetName() + TEXT("_settings");
+			PackageName = FPackageName::GetLongPackagePath(StaticMesh->GetOutermost()->GetName()) + TEXT("/") + StaticMesh->GetName() + TEXT("_FoliageType");
 		}
 		
-		TSharedRef<SDlgPickAssetPath> SettingDlg =
+		TSharedRef<SDlgPickAssetPath> SaveFoliageTypeDialog =
 			SNew(SDlgPickAssetPath)
-			.Title(LOCTEXT("SettingsDialogTitle", "Choose Location for Foliage Settings Asset"))
+			.Title(LOCTEXT("SaveFoliageTypeDialogTitle", "Choose Location for Foliage Type Asset"))
 			.DefaultAssetPath(FText::FromString(PackageName));
 
-		if (SettingDlg->ShowModal() != EAppReturnType::Cancel)
+		if (SaveFoliageTypeDialog->ShowModal() != EAppReturnType::Cancel)
 		{
-			PackageName = SettingDlg->GetFullAssetPath().ToString();
+			PackageName = SaveFoliageTypeDialog->GetFullAssetPath().ToString();
 			UPackage* Package = CreatePackage(nullptr, *PackageName);
-			SettingsToSave = Cast<UFoliageType>(StaticDuplicateObject(Settings, Package, *FPackageName::GetLongPackageAssetName(PackageName)));
-			SettingsToSave->SetFlags(RF_Standalone | RF_Public);
-			SettingsToSave->Modify();
+			TypeToSave = Cast<UFoliageType>(StaticDuplicateObject(InFoliageType, Package, *FPackageName::GetLongPackageAssetName(PackageName)));
+			TypeToSave->SetFlags(RF_Standalone | RF_Public);
+			TypeToSave->Modify();
 
 			// Notify the asset registry
-			FAssetRegistryModule::AssetCreated(SettingsToSave);
+			FAssetRegistryModule::AssetCreated(TypeToSave);
 
-			ReplaceSettingsObject(Settings, SettingsToSave);
+			ReplaceSettingsObject(InFoliageType, TypeToSave);
 		}
 	}
 	else
 	{
-		SettingsToSave = Settings;
+		TypeToSave = InFoliageType;
 	}
 	
 	// Save to disk
-	if (SettingsToSave)
+	if (TypeToSave)
 	{
 		TArray<UPackage*> PackagesToSave; 
-		PackagesToSave.Add(SettingsToSave->GetOutermost());
+		PackagesToSave.Add(TypeToSave->GetOutermost());
 		const bool bCheckDirty = false;
 		const bool bPromptToSave = false;
 		FEditorFileUtils::PromptForCheckoutAndSave(PackagesToSave, bCheckDirty, bPromptToSave);
 	}
 		
-	return SettingsToSave;
+	return TypeToSave;
 }
 
 /** Reapply cluster settings to all the instances */
diff --git a/Engine/Source/Editor/FoliageEdit/Private/FoliageEdMode.h b/Engine/Source/Editor/FoliageEdit/Private/FoliageEdMode.h
--- a/Engine/Source/Editor/FoliageEdit/Private/FoliageEdMode.h
+++ b/Engine/Source/Editor/FoliageEdit/Private/FoliageEdMode.h
@@ -326,8 +326,8 @@ public:
 	/** Replace the settings object for this static mesh with the one specified */
 	void ReplaceSettingsObject(UFoliageType* OldSettings, UFoliageType* NewSettings);
 
-	/** Save the settings object as an asset */
-	UFoliageType* SaveSettingsObject(UFoliageType* Settings);
+	/** Save the foliage type object. If it isn't an asset, will prompt the user for a location to save the new asset. */
+	UFoliageType* SaveFoliageTypeObject(UFoliageType* Settings);
 
 	/** Set/Clear selection for foliage instances of specific type  */
 	void SelectInstances(const UFoliageType* Settings, bool bSelect);
diff --git a/Engine/Source/Editor/FoliageEdit/Private/FoliageEditActions.cpp b/Engine/Source/Editor/FoliageEdit/Private/FoliageEditActions.cpp
--- a/Engine/Source/Editor/FoliageEdit/Private/FoliageEditActions.cpp
+++ b/Engine/Source/Editor/FoliageEdit/Private/FoliageEditActions.cpp
@@ -20,11 +20,11 @@ void FFoliageEditCommands::RegisterCommands()
 	UI_COMMAND( SetPaintSettings, "Show Painting settings", "Show painting settings.", EUserInterfaceActionType::ToggleButton, FInputChord() );
 	UI_COMMAND( SetClusterSettings, "Show Instance settings", "Show settings for placed instances.", EUserInterfaceActionType::ToggleButton, FInputChord() );
 
-	UI_COMMAND( RemoveFoliageType, "Remove", "Remove all associated foliage instances", EUserInterfaceActionType::Button, FInputChord() );
+	UI_COMMAND( RemoveFoliageType, "Remove", "Remove this foliage type from the palette. Removes all associated instances as well.", EUserInterfaceActionType::Button, FInputChord() );
 	UI_COMMAND( ShowFoliageTypeInCB, "Show in Content Browser", "Show asset in Content Browser.", EUserInterfaceActionType::Button, FInputChord() );
-	UI_COMMAND( SelectAllInstances, "Select All Instances", "Select all instances of foliage type.", EUserInterfaceActionType::Button, FInputChord());
-	UI_COMMAND( DeselectAllInstances, "Deselect All Instances", "Deselect all instances of foliage type.", EUserInterfaceActionType::Button, FInputChord());
-	UI_COMMAND( SelectInvalidInstances, "Select Invalid Instances", "Select all instances that are off ground.", EUserInterfaceActionType::Button, FInputChord());
+	UI_COMMAND( SelectAllInstances, "Select All Instances", "Select all instances of this foliage type (must be in a selection mode).", EUserInterfaceActionType::Button, FInputChord());
+	UI_COMMAND( DeselectAllInstances, "Deselect All Instances", "Deselect all instances of this foliage type (must be in a selection mode).", EUserInterfaceActionType::Button, FInputChord());
+	UI_COMMAND( SelectInvalidInstances, "Select Invalid Instances", "Select all instances of this foliage type that are off ground (must be in a selection mode).", EUserInterfaceActionType::Button, FInputChord());
 }
 
 #undef LOCTEXT_NAMESPACE
\ No newline at end of file
diff --git a/Engine/Source/Editor/FoliageEdit/Private/FoliagePaletteItem.cpp b/Engine/Source/Editor/FoliageEdit/Private/FoliagePaletteItem.cpp
new file mode 100644
--- /dev/null
+++ b/Engine/Source/Editor/FoliageEdit/Private/FoliagePaletteItem.cpp
@@ -0,0 +1,360 @@
+// Copyright 1998-2015 Epic Games, Inc. All Rights Reserved.
+
+#include "UnrealEd.h"
+
+#include "FoliagePaletteItem.h"
+#include "SFoliagePalette.h"
+#include "FoliageEdMode.h"
+#include "FoliageType.h"
+
+#include "LevelEditor.h"
+#include "Editor/UnrealEd/Public/AssetThumbnail.h"
+#include "Engine/StaticMesh.h"
+
+#define LOCTEXT_NAMESPACE "FoliageEd_Mode"
+
+FFoliagePaletteItemModel::FFoliagePaletteItemModel(FFoliageMeshUIInfoPtr InTypeInfo, TSharedRef<SFoliagePalette> InFoliagePalette, FEdModeFoliage* InFoliageEditMode)
+	: TypeInfo(InTypeInfo)
+	, FoliagePalette(InFoliagePalette)
+	, FoliageEditMode(InFoliageEditMode)
+{
+	FLevelEditorModule& LevelEditorModule = FModuleManager::LoadModuleChecked<FLevelEditorModule>("LevelEditor");
+	TSharedPtr<FAssetThumbnailPool> ThumbnailPool = LevelEditorModule.GetFirstLevelEditor()->GetThumbnailPool();
+	FAssetData AssetData = FAssetData(TypeInfo->Settings->GetStaticMesh());
+
+	TSharedPtr< FAssetThumbnail > Thumbnail = MakeShareable(new FAssetThumbnail(AssetData, 128, 128, ThumbnailPool));
+	FAssetThumbnailConfig ThumbnailConfig;
+	if (TypeInfo->Settings->IsAsset())
+	{
+		//@todo: access the foliage type asset type action and get the color from that
+		ThumbnailConfig.AssetTypeColorOverride = FColor(128, 192, 255);
+	}
+
+	ThumbnailWidget = Thumbnail->MakeThumbnailWidget(ThumbnailConfig);
+}
+
+const TSharedPtr<SFoliagePalette>& FFoliagePaletteItemModel::GetFoliagePalette() const
+{
+	return FoliagePalette;
+}
+
+const UFoliageType* FFoliagePaletteItemModel::GetFoliageType() const
+{
+	return TypeInfo->Settings;
+}
+
+TSharedRef<SWidget> FFoliagePaletteItemModel::GetThumbnailWidget() const
+{
+	return ThumbnailWidget.ToSharedRef();
+}
+
+TSharedRef<SToolTip> FFoliagePaletteItemModel::CreateTooltipWidget() const
+{
+	return 
+		SNew(SToolTip)
+		.TextMargin(1)
+		.BorderImage(FEditorStyle::GetBrush("ContentBrowser.TileViewTooltip.ToolTipBorder"))
+		.Visibility(this, &FFoliagePaletteItemModel::GetTooltipVisibility)
+		[
+			SNew(SBorder)
+			.Padding(3.f)
+			.BorderImage(FEditorStyle::GetBrush("ContentBrowser.TileViewTooltip.NonContentBorder"))
+			[
+				SNew(SVerticalBox)
+
+				+ SVerticalBox::Slot()
+				.AutoHeight()
+				[
+					SNew(SBorder)
+					.Padding(FMargin(6.f))
+					.HAlign(HAlign_Left)
+					.BorderImage(FEditorStyle::GetBrush("ContentBrowser.TileViewTooltip.ContentBorder"))
+					[
+						SNew(STextBlock)
+						.Text(this, &FFoliagePaletteItemModel::GetFoliageTypeNameText)
+						.Font(FEditorStyle::GetFontStyle("ContentBrowser.TileViewTooltip.NameFont"))
+					]
+				]
+
+				+ SVerticalBox::Slot()
+				.AutoHeight()
+				.Padding(FMargin(0.f, 3.f, 0.f, 0.f))
+				[
+					SNew(SBorder)
+					.Visibility(this, &FFoliagePaletteItemModel::GetTooltipThumbnailVisibility)
+					.Padding(6.f)
+					.HAlign(HAlign_Left)
+					.BorderImage(FEditorStyle::GetBrush("ContentBrowser.TileViewTooltip.ContentBorder"))
+					[
+						SNew(SBox)
+						.HeightOverride(64.f)
+						.WidthOverride(64.f)
+						[
+							GetThumbnailWidget()
+						]
+					]
+				]
+			]
+		];
+}
+
+TSharedRef<SCheckBox> FFoliagePaletteItemModel::CreateActivationCheckBox(TAttribute<bool> IsItemWidgetSelected, TAttribute<EVisibility> InVisibility)
+{
+	return
+		SNew(SCheckBox)
+		.Padding(0.f)
+		.OnCheckStateChanged(this, &FFoliagePaletteItemModel::HandleCheckStateChanged, IsItemWidgetSelected)
+		.Visibility(InVisibility)
+		.IsChecked(this, &FFoliagePaletteItemModel::GetCheckBoxState)
+		.ToolTipText(LOCTEXT("TileCheckboxTooltip", "Check to activate the currently selected types in the palette"));
+}
+
+TSharedRef<SButton> FFoliagePaletteItemModel::CreateSaveAssetButton(TAttribute<EVisibility> InVisibility)
+{
+	return 
+		SNew(SButton)
+		.ContentPadding(0)
+		.ButtonStyle(FEditorStyle::Get(), "ToggleButton")
+		.Visibility(InVisibility)
+		.IsEnabled(this, &FFoliagePaletteItemModel::IsSaveEnabled)
+		.OnClicked(this, &FFoliagePaletteItemModel::HandleSaveAsset)
+		.ToolTipText(LOCTEXT("SaveButtonToolTip", "Save foliage asset"))
+		.HAlign(HAlign_Center)
+		.VAlign(VAlign_Center)
+		.Content()
+		[
+			SNew(SImage)
+			.Image(FEditorStyle::GetBrush("Level.SaveIcon16x"))
+		];
+}
+
+FText FFoliagePaletteItemModel::GetFoliageTypeNameText() const
+{
+	if (TypeInfo->Settings->IsAsset())
+	{
+		return FText::FromName(TypeInfo->Settings->GetFName());
+	}
+	else
+	{
+		return FText::FromName(TypeInfo->Settings->GetStaticMesh()->GetFName());
+	}
+}
+
+FText FFoliagePaletteItemModel::GetInstanceCountText() const
+{
+	const int32	InstanceCountTotal = TypeInfo->InstanceCountTotal;
+	const int32	InstanceCountCurrentLevel = TypeInfo->InstanceCountCurrentLevel;
+
+	if (InstanceCountCurrentLevel != InstanceCountTotal)
+	{
+		return FText::Format(LOCTEXT("InstanceCount", "{0} ({1})"), FText::AsNumber(InstanceCountCurrentLevel), FText::AsNumber(InstanceCountTotal));
+	}
+	else
+	{
+		return FText::AsNumber(InstanceCountCurrentLevel);
+	}
+}
+
+void FFoliagePaletteItemModel::HandleCheckStateChanged(const ECheckBoxState NewCheckedState, TAttribute<bool> IsItemWidgetSelected)
+{
+	if (!IsItemWidgetSelected.IsSet()) { return; }
+
+	const bool bShouldActivate = NewCheckedState == ECheckBoxState::Checked;
+	if (!IsItemWidgetSelected.Get())
+	{
+		if (TypeInfo->Settings->IsSelected != bShouldActivate)
+		{
+			TypeInfo->Settings->Modify();
+			TypeInfo->Settings->IsSelected = bShouldActivate;
+		}
+	}
+	else
+	{
+		FoliagePalette->ActivateAllSelectedTypes(bShouldActivate);
+	}
+}
+
+ECheckBoxState FFoliagePaletteItemModel::GetCheckBoxState() const
+{
+	return TypeInfo->Settings->IsSelected ? ECheckBoxState::Checked : ECheckBoxState::Unchecked;
+}
+
+bool FFoliagePaletteItemModel::IsSaveEnabled() const
+{
+	UFoliageType* FoliageType = TypeInfo->Settings;
+
+	// Saving is enabled for non-assets and dirty assets
+	return (!FoliageType->IsAsset() || FoliageType->GetOutermost()->IsDirty());
+}
+
+FReply FFoliagePaletteItemModel::HandleSaveAsset()
+{
+	UFoliageType* SavedSettings = FoliageEditMode->SaveFoliageTypeObject(TypeInfo->Settings);
+	if (SavedSettings)
+	{
+		TypeInfo->Settings = SavedSettings;
+	}
+
+	return FReply::Handled();
+}
+
+EVisibility FFoliagePaletteItemModel::GetTooltipVisibility() const
+{
+	return FoliagePalette->ShouldShowTooltips() ? EVisibility::SelfHitTestInvisible : EVisibility::Collapsed;
+}
+
+EVisibility FFoliagePaletteItemModel::GetTooltipThumbnailVisibility() const
+{
+	return FoliagePalette->IsActiveViewMode(EPaletteViewMode::Tree) ? EVisibility::SelfHitTestInvisible : EVisibility::Collapsed;
+}
+
+////////////////////////////////////////////////
+// SFoliagePaletteItemTile
+////////////////////////////////////////////////
+
+void SFoliagePaletteItemTile::Construct(const FArguments& InArgs, TSharedRef<STableViewBase> InOwnerTableView, TSharedPtr<FFoliagePaletteItemModel>& InModel)
+{
+	Model = InModel;
+
+	auto IsSelectedGetter = TAttribute<bool>::Create(TAttribute<bool>::FGetter::CreateSP(this, &SFoliagePaletteItemTile::IsSelected));
+	auto CheckBoxVisibility = TAttribute<EVisibility>::Create(TAttribute<EVisibility>::FGetter::CreateSP(this, &SFoliagePaletteItemTile::GetCheckBoxVisibility));
+	auto SaveButtonVisibility = TAttribute<EVisibility>::Create(TAttribute<EVisibility>::FGetter::CreateSP(this, &SFoliagePaletteItemTile::GetSaveButtonVisibility));
+
+	STableRow<FFoliageMeshUIInfoPtr>::Construct(
+		STableRow<FFoliageMeshUIInfoPtr>::FArguments()
+		.Style(FEditorStyle::Get(), "ContentBrowser.AssetListView.TableRow")
+		.Padding(1.f)
+		.Content()
+		[
+			SNew(SOverlay)
+			.ToolTip(Model->CreateTooltipWidget())
+			
+			+ SOverlay::Slot()
+			[
+				SNew(SBorder)
+				.Padding(4.f)
+				.BorderImage(FEditorStyle::GetBrush("ContentBrowser.ThumbnailShadow"))
+				.ForegroundColor(FLinearColor::White)
+				.ColorAndOpacity(this, &SFoliagePaletteItemTile::GetTileColorAndOpacity)
+				[
+					Model->GetThumbnailWidget()
+				]
+			]
+			
+			+ SOverlay::Slot()
+			.HAlign(HAlign_Left)
+			.VAlign(VAlign_Top)
+			.Padding(FMargin(3.f))
+			[
+				SNew(SBorder)
+				.BorderImage(FEditorStyle::GetBrush("ContentBrowser.ThumbnailShadow"))
+				.BorderBackgroundColor(FLinearColor::Black)
+				.ForegroundColor(FLinearColor::White)
+				.Padding(3.f)
+				[
+					Model->CreateActivationCheckBox(IsSelectedGetter, CheckBoxVisibility)
+				]
+			]
+
+			+ SOverlay::Slot()
+			.HAlign(HAlign_Right)
+			.VAlign(VAlign_Top)
+			.Padding(FMargin(3.f))
+			[
+				Model->CreateSaveAssetButton(SaveButtonVisibility)
+			]
+
+			+ SOverlay::Slot()
+
+			+ SOverlay::Slot()
+			.HAlign(HAlign_Right)
+			.VAlign(VAlign_Bottom)
+			.Padding(FMargin(6.f, 8.f))
+			[
+				SNew(STextBlock)
+				.Text(Model.Get(), &FFoliagePaletteItemModel::GetInstanceCountText)
+				.ShadowOffset(FVector2D(1.f, 1.f))
+				.ColorAndOpacity(FLinearColor(.85f, .85f, .85f, 1.f))
+			]
+		], InOwnerTableView);
+}
+
+FLinearColor SFoliagePaletteItemTile::GetTileColorAndOpacity() const
+{
+	float Alpha = Model->GetFoliageType()->IsSelected ? 1.f : 0.5f;
+	return FLinearColor(1.f, 1.f, 1.f, Alpha);
+}
+
+EVisibility SFoliagePaletteItemTile::GetCheckBoxVisibility() const
+{
+	return IsHovered() || (IsSelected() && Model->GetFoliagePalette()->AnySelectedTileHovered()) ? EVisibility::Visible : EVisibility::Collapsed;
+}
+
+EVisibility SFoliagePaletteItemTile::GetSaveButtonVisibility() const
+{
+	return IsHovered() ? EVisibility::Visible : EVisibility::Collapsed;
+}
+
+////////////////////////////////////////////////
+// SFoliagePaletteItemRow
+////////////////////////////////////////////////
+
+void SFoliagePaletteItemRow::Construct(const FArguments& InArgs, TSharedRef<STableViewBase> InOwnerTableView, TSharedPtr<FFoliagePaletteItemModel>& InModel)
+{
+	Model = InModel;
+	SMultiColumnTableRow<FFoliageMeshUIInfoPtr>::Construct(FSuperRowType::FArguments(), InOwnerTableView);
+
+	SetToolTip(Model->CreateTooltipWidget());
+}
+
+TSharedRef<SWidget> SFoliagePaletteItemRow::GenerateWidgetForColumn(const FName& ColumnName)
+{
+	TSharedPtr<SWidget> TableRowContent = SNullWidget::NullWidget;
+
+	if (ColumnName == FoliagePaletteTreeColumns::ColumnID_ToggleActive)
+	{
+		auto IsSelectedGetter = TAttribute<bool>::Create(TAttribute<bool>::FGetter::CreateSP(this, &SFoliagePaletteItemTile::IsSelected));
+		TableRowContent = Model->CreateActivationCheckBox(IsSelectedGetter);
+	} 
+	else if (ColumnName == FoliagePaletteTreeColumns::ColumnID_Type)
+	{
+		TableRowContent =
+			SNew(SHorizontalBox)
+			
+			+SHorizontalBox::Slot()
+			.AutoWidth()
+			[
+				SNew(SExpanderArrow, SharedThis(this))
+			]
+			
+			+SHorizontalBox::Slot()
+			.VAlign(VAlign_Center)
+			.AutoWidth()
+			[
+				SNew(STextBlock)
+				.Text(Model.Get(), &FFoliagePaletteItemModel::GetFoliageTypeNameText)
+			];
+	}
+	else if (ColumnName == FoliagePaletteTreeColumns::ColumnID_InstanceCount)
+	{
+		TableRowContent =
+			SNew(SHorizontalBox)
+			
+			+SHorizontalBox::Slot()
+			.Padding(FMargin(10,1,0,1))
+			.VAlign(VAlign_Center)
+			.AutoWidth()
+			[
+				SNew(STextBlock)
+				.Text(Model.Get(), &FFoliagePaletteItemModel::GetInstanceCountText)
+			];
+	}
+	else if (ColumnName == FoliagePaletteTreeColumns::ColumnID_Save)
+	{
+		TableRowContent = Model->CreateSaveAssetButton();
+	}
+
+	return TableRowContent.ToSharedRef();
+}
+
+#undef LOCTEXT_NAMESPACE
\ No newline at end of file
diff --git a/Engine/Source/Editor/FoliageEdit/Private/FoliagePaletteItem.h b/Engine/Source/Editor/FoliageEdit/Private/FoliagePaletteItem.h
new file mode 100644
--- /dev/null
+++ b/Engine/Source/Editor/FoliageEdit/Private/FoliagePaletteItem.h
@@ -0,0 +1,107 @@
+// Copyright 1998-2015 Epic Games, Inc. All Rights Reserved.
+
+#pragma once
+
+class SToolTip;
+class SFoliagePalette;
+
+class FEdModeFoliage;
+struct FFoliageMeshUIInfo;
+
+typedef TSharedPtr<FFoliageMeshUIInfo> FFoliageMeshUIInfoPtr; //should match typedef in FoliageEdMode.h
+
+namespace FoliagePaletteTreeColumns
+{
+	/** IDs for list columns */
+	static const FName ColumnID_ToggleActive("Toggle");
+	static const FName ColumnID_Type("Type");
+	static const FName ColumnID_InstanceCount("InstanceCount");
+	static const FName ColumnID_Save("Save");
+};
+
+class FFoliagePaletteItemModel : public TSharedFromThis<FFoliagePaletteItemModel>
+{
+public:
+	FFoliagePaletteItemModel(FFoliageMeshUIInfoPtr InTypeInfo, TSharedRef<SFoliagePalette> InFoliagePalette, FEdModeFoliage* InFoliageEditMode);
+
+	/** @return The foliage palette that contains the item */
+	const TSharedPtr<SFoliagePalette>& GetFoliagePalette() const;
+
+	const UFoliageType* GetFoliageType() const;
+
+	/** @return The thumbnail widget for this item */
+	TSharedRef<SWidget> GetThumbnailWidget() const;
+
+	/** @return The tooltip widget for this item */
+	TSharedRef<class SToolTip> CreateTooltipWidget() const;
+
+	/** @return The checkbox widget for activating/deactivating this type in the palette */
+	TSharedRef<class SCheckBox> CreateActivationCheckBox(TAttribute<bool> IsItemWidgetSelected, TAttribute<EVisibility> InVisibility = TAttribute<EVisibility>());
+
+	/** @return The save asset button */
+	TSharedRef<class SButton> CreateSaveAssetButton(TAttribute<EVisibility> InVisibility = TAttribute<EVisibility>());
+
+	/** Gets the name to display of the foliage type object as text */
+	FText GetFoliageTypeNameText() const;
+
+	/** Gets the number of instances of this foliage type exist in the level */
+	FText GetInstanceCountText() const;
+
+private:
+	/** Handles the change in activation of the item in the palette */
+	void HandleCheckStateChanged(const ECheckBoxState NewCheckedState, TAttribute<bool> IsItemWidgetSelected);
+
+	/** Gets whether the foliage type is active in the palette  */
+	ECheckBoxState GetCheckBoxState() const;
+
+	/** Whether the save asset button is enabled */
+	bool IsSaveEnabled() const;
+
+	/** Saves the foliage type */
+	FReply HandleSaveAsset();
+
+	/** Gets the visibility of the entire tooltip */
+	EVisibility GetTooltipVisibility() const;
+
+	/** Gets the visibility of the thumbnail in the tooltip */
+	EVisibility GetTooltipThumbnailVisibility() const;
+
+private:
+	TSharedPtr<SWidget> ThumbnailWidget;
+
+	FFoliageMeshUIInfoPtr TypeInfo;
+	TSharedPtr<SFoliagePalette> FoliagePalette;
+	FEdModeFoliage* FoliageEditMode;
+};
+
+/** A tile representing a foliage type in the palette */
+class SFoliagePaletteItemTile : public STableRow<FFoliageMeshUIInfoPtr>
+{
+public:
+	SLATE_BEGIN_ARGS(SFoliagePaletteItemTile) {}
+	SLATE_END_ARGS()
+
+	void Construct(const FArguments& InArgs, TSharedRef<STableViewBase> InOwnerTableView, TSharedPtr<FFoliagePaletteItemModel>& InModel);
+
+private:
+	FLinearColor GetTileColorAndOpacity() const;
+	EVisibility GetCheckBoxVisibility() const;
+	EVisibility GetSaveButtonVisibility() const;
+
+private:
+	TSharedPtr<FFoliagePaletteItemModel> Model;
+};
+
+/** A tree row representing a foliage type in the palette */
+class SFoliagePaletteItemRow : public SMultiColumnTableRow<FFoliageMeshUIInfoPtr>
+{
+public:
+	SLATE_BEGIN_ARGS(SFoliagePaletteItemRow) {}
+	SLATE_END_ARGS()
+
+	void Construct(const FArguments& InArgs, TSharedRef<STableViewBase> InOwnerTableView, TSharedPtr<FFoliagePaletteItemModel>& InModel);
+	virtual TSharedRef<SWidget> GenerateWidgetForColumn(const FName& ColumnName) override;
+
+private:
+	TSharedPtr<FFoliagePaletteItemModel> Model;
+};
\ No newline at end of file
diff --git a/Engine/Source/Editor/FoliageEdit/Private/SFoliageEdit.cpp b/Engine/Source/Editor/FoliageEdit/Private/SFoliageEdit.cpp
--- a/Engine/Source/Editor/FoliageEdit/Private/SFoliageEdit.cpp
+++ b/Engine/Source/Editor/FoliageEdit/Private/SFoliageEdit.cpp
@@ -2,29 +2,23 @@
 
 #include "UnrealEd.h"
 
-#include "SScaleBox.h"
-
-#include "FoliageType.h"
+#include "FoliageType_InstancedStaticMesh.h"
 #include "FoliageEdMode.h"
 #include "SFoliageEdit.h"
 #include "FoliageEditActions.h"
-#include "Editor/UnrealEd/Public/AssetThumbnail.h"
-#include "Editor/UnrealEd/Public/DragAndDrop/AssetDragDropOp.h"
-#include "Editor/UnrealEd/Public/ScopedTransaction.h"
-#include "Editor/ContentBrowser/Public/ContentBrowserModule.h"
-#include "Editor/UnrealEd/Public/AssetSelection.h"
 #include "Editor/IntroTutorials/Public/IIntroTutorials.h"
-#include "Editor/PropertyEditor/Public/IDetailsView.h"
-#include "Editor/PropertyEditor/Public/PropertyEditorModule.h"
-#include "Engine/StaticMesh.h"
-#include "SFoliageEditMeshRow.h"
-#include "FoliageTypePaintingCustomization.h"
-#include "DetailLayoutBuilder.h"
 #include "SNumericEntryBox.h"
+#include "SFoliagePalette.h"
+#include "SScaleBox.h"
+
+
+// Possibly not needed (move over to SFoliagePalette)
+#include "Editor/UnrealEd/Public/AssetSelection.h"
+#include "Editor/UnrealEd/Public/ScopedTransaction.h"
+
+// Seemingly totally unneeded
+//#include "DetailLayoutBuilder.h"
 
-#include "FoliageType_InstancedStaticMesh.h"
-#include "Editor/PropertyEditor/Public/PropertyCustomizationHelpers.h"
-#include "LevelEditor.h"
 
 #define LOCTEXT_NAMESPACE "FoliageEd_Mode"
 
@@ -94,24 +88,12 @@ BEGIN_SLATE_FUNCTION_BUILD_OPTIMIZATION
 void SFoliageEdit::Construct(const FArguments& InArgs)
 {
 	FoliageEditMode = (FEdModeFoliage*)GLevelEditorModeTools().GetActiveMode(FBuiltinEditorModes::EM_Foliage);
-	FoliageEditMode->OnToolChanged.AddSP(this, &SFoliageEdit::HandleOnToolChanged);
 
 	FFoliageEditCommands::Register();
-
 	UICommandList = MakeShareable(new FUICommandList);
-
 	BindCommands();
 
-	FPropertyEditorModule& PropertyModule = FModuleManager::LoadModuleChecked<FPropertyEditorModule>("PropertyEditor");
-	FDetailsViewArgs Args(false, false, false, FDetailsViewArgs::HideNameArea, true);
-	Args.bShowActorLabel = false;
-	MeshDetailsWidget = PropertyModule.CreateDetailView(Args);
-
-	// We want to use our own customization for UFoliageType
-	MeshDetailsWidget->UnregisterInstancedCustomPropertyLayout(UFoliageType::StaticClass());
-	MeshDetailsWidget->RegisterInstancedCustomPropertyLayout(UFoliageType::StaticClass(), 
-		FOnGetDetailCustomizationInstance::CreateStatic(&FFoliageTypePaintingCustomization::MakeInstance, FoliageEditMode)
-		);
+	IIntroTutorials& IntroTutorials = FModuleManager::LoadModuleChecked<IIntroTutorials>(TEXT("IntroTutorials"));
 
 	// Everything (or almost) uses this padding, change it to expand the padding.
 	FMargin StandardPadding(6.f, 3.f);
@@ -124,357 +106,317 @@ void SFoliageEdit::Construct(const FArguments& InArgs)
 
 	this->ChildSlot
 	[
-		SNew(SOverlay)
+		SNew(SVerticalBox)
 
-		+ SOverlay::Slot()
+		+ SVerticalBox::Slot()
+		.AutoHeight()
 		[
-			SNew(SVerticalBox)
-
-			+ SVerticalBox::Slot()
-			.AutoHeight()
-			.Padding(StandardPadding)
+			SNew(SHorizontalBox)
+			+ SHorizontalBox::Slot()
+			.AutoWidth()
+			.Padding(1.f, 5.f, 0.f, 5.f)
 			[
 				BuildToolBar()
 			]
 
-			// Brush Size
-			+ SVerticalBox::Slot()
-			.AutoHeight()
+			+ SHorizontalBox::Slot()
+			.Padding(0.f, 2.f, 2.f, 0.f)
 			[
-				SNew(SHorizontalBox)
-				.ToolTipText(LOCTEXT("BrushSize_Tooltip", "The size of the foliage brush"))
-				.Visibility(this, &SFoliageEdit::GetVisibility_Radius)
-
-				+ SHorizontalBox::Slot()
-				.Padding(StandardLeftPadding)
-				.FillWidth(1.0f)
-				.VAlign(VAlign_Center)
-				[
-					SNew(STextBlock)
-					.Text(LOCTEXT("BrushSize", "Brush Size"))
-					.Font(StandardFont)
-				]
-				+ SHorizontalBox::Slot()
-				.Padding(StandardRightPadding)
-				.FillWidth(2.0f)
-				.MaxWidth(100.f)
-				.VAlign(VAlign_Center)
+				SNew(SBorder)
+				.BorderImage(FEditorStyle::GetBrush("ToolPanel.DarkGroupBorder"))
+				.Padding(StandardPadding)
 				[
-					SNew(SNumericEntryBox<float>)
-					.Font(StandardFont)
-					.AllowSpin(true)
-					.MinValue(0.0f)
-					.MaxValue(65536.0f)
-					.MaxSliderValue(8192.0f)
-					.SliderExponent(3.0f)
-					.Value(this, &SFoliageEdit::GetRadius)
-					.OnValueChanged(this, &SFoliageEdit::SetRadius)
-				]
-			]
+					SNew(SVerticalBox)
 
-			// Paint Density
-			+ SVerticalBox::Slot()
-			.AutoHeight()
-			[
-				SNew(SHorizontalBox)
-				.ToolTipText(LOCTEXT("PaintDensity_Tooltip", "The density of foliage to paint. This is a multiplier for the individual foliage type's density specifier."))
-				.Visibility(this, &SFoliageEdit::GetVisibility_PaintDensity)
-
-				+ SHorizontalBox::Slot()
-				.Padding(StandardLeftPadding)
-				.FillWidth(1.0f)
-				.VAlign(VAlign_Center)
-				[
-					SNew(STextBlock)
-					.Text(LOCTEXT("PaintDensity", "Paint Density"))
-					.Font(StandardFont)
-				]
-				+ SHorizontalBox::Slot()
-				.Padding(StandardRightPadding)
-				.FillWidth(2.0f)
-				.MaxWidth(100.f)
-				.VAlign(VAlign_Center)
-				[
-					SNew(SNumericEntryBox<float>)
-					.Font(StandardFont)
-					.AllowSpin(true)
-					.MinValue(0.0f)
-					.MaxValue(2.0f)
-					.MaxSliderValue(1.0f)
-					.Value(this, &SFoliageEdit::GetPaintDensity)
-					.OnValueChanged(this, &SFoliageEdit::SetPaintDensity)
-				]
-			]
+					// Active Tool Title
+					+ SVerticalBox::Slot()
+					.AutoHeight()
+					[
+						SNew(SHorizontalBox)
 
-			// Erase Density
-			+ SVerticalBox::Slot()
-			.AutoHeight()
-			[
-				SNew(SHorizontalBox)
-				.ToolTipText(LOCTEXT("EraseDensity_Tooltip", "The density of foliage to leave behind when erasing with the Shift key held. 0 will remove all foliage."))
-				.Visibility(this, &SFoliageEdit::GetVisibility_EraseDensity)
-
-				+ SHorizontalBox::Slot()
-				.Padding(StandardLeftPadding)
-				.FillWidth(1.0f)
-				.VAlign(VAlign_Center)
-				[
-					SNew(STextBlock)
-					.Text(LOCTEXT("EraseDensity", "Erase Density"))
-					.Font(StandardFont)
-				]
-				+ SHorizontalBox::Slot()
-				.Padding(StandardRightPadding)
-				.FillWidth(2.0f)
-				.MaxWidth(100.f)
-				.VAlign(VAlign_Center)
-				[
-					SNew(SNumericEntryBox<float>)
-					.Font(StandardFont)
-					.AllowSpin(true)
-					.MinValue(0.0f)
-					.MaxValue(2.0f)
-					.MaxSliderValue(1.0f)
-					.Value(this, &SFoliageEdit::GetEraseDensity)
-					.OnValueChanged(this, &SFoliageEdit::SetEraseDensity)
-				]
-			]
+						+ SHorizontalBox::Slot()
+						.Padding(StandardLeftPadding)
+						.HAlign(HAlign_Left)
+						[
+							SNew(STextBlock)
+							.Text(this, &SFoliageEdit::GetActiveToolName)
+							.TextStyle(FEditorStyle::Get(), "FoliageEditMode.ActiveToolName.Text")
+						]
 
-			// Filter
-			+ SVerticalBox::Slot()
-			.AutoHeight()
-			[
-				SNew(SHorizontalBox)
-				.Visibility(this, &SFoliageEdit::GetVisibility_Filters)
+						+ SHorizontalBox::Slot()
+						.Padding(StandardRightPadding)
+						.HAlign(HAlign_Right)
+						.VAlign(VAlign_Center)
+						.AutoWidth()
+						[
+							// Tutorial link
+							IntroTutorials.CreateTutorialsWidget(TEXT("FoliageMode"))
+						]
+					]
 
-				+ SHorizontalBox::Slot()
-				.Padding(StandardLeftPadding)
-				.VAlign(VAlign_Center)
-				.FillWidth(1.0f)
-				[
-					SNew(STextBlock)
-					.Text(this, &SFoliageEdit::GetFilterText)
-					.ToolTipText(this, &SFoliageEdit::GetFilterTooltipText)
-					.Font(StandardFont)
-				]
+					// Brush Size
+					+ SVerticalBox::Slot()
+					.AutoHeight()
+					[
+						SNew(SHorizontalBox)
+						.ToolTipText(LOCTEXT("BrushSize_Tooltip", "The size of the foliage brush"))
+						.Visibility(this, &SFoliageEdit::GetVisibility_Radius)
+
+						+ SHorizontalBox::Slot()
+						.Padding(StandardLeftPadding)
+						.FillWidth(1.0f)
+						.VAlign(VAlign_Center)
+						[
+							SNew(STextBlock)
+							.Text(LOCTEXT("BrushSize", "Brush Size"))
+							.Font(StandardFont)
+						]
+						+ SHorizontalBox::Slot()
+						.Padding(StandardRightPadding)
+						.FillWidth(2.0f)
+						.MaxWidth(100.f)
+						.VAlign(VAlign_Center)
+						[
+							SNew(SNumericEntryBox<float>)
+							.Font(StandardFont)
+							.AllowSpin(true)
+							.MinValue(0.0f)
+							.MaxValue(65536.0f)
+							.MaxSliderValue(8192.0f)
+							.SliderExponent(3.0f)
+							.Value(this, &SFoliageEdit::GetRadius)
+							.OnValueChanged(this, &SFoliageEdit::SetRadius)
+						]
+					]
 
-				+ SHorizontalBox::Slot()
-				.VAlign(VAlign_Center)
-				.FillWidth(2.0f)
-				[
-					SNew(SWrapBox)
-					.UseAllottedWidth(true)
+					// Paint Density
+					+ SVerticalBox::Slot()
+					.AutoHeight()
+					[
+						SNew(SHorizontalBox)
+						.ToolTipText(LOCTEXT("PaintDensity_Tooltip", "The density of foliage to paint. This is a multiplier for the individual foliage type's density specifier."))
+						.Visibility(this, &SFoliageEdit::GetVisibility_PaintDensity)
+
+						+ SHorizontalBox::Slot()
+						.Padding(StandardLeftPadding)
+						.FillWidth(1.0f)
+						.VAlign(VAlign_Center)
+						[
+							SNew(STextBlock)
+							.Text(LOCTEXT("PaintDensity", "Paint Density"))
+							.Font(StandardFont)
+						]
+						+ SHorizontalBox::Slot()
+						.Padding(StandardRightPadding)
+						.FillWidth(2.0f)
+						.MaxWidth(100.f)
+						.VAlign(VAlign_Center)
+						[
+							SNew(SNumericEntryBox<float>)
+							.Font(StandardFont)
+							.AllowSpin(true)
+							.MinValue(0.0f)
+							.MaxValue(2.0f)
+							.MaxSliderValue(1.0f)
+							.Value(this, &SFoliageEdit::GetPaintDensity)
+							.OnValueChanged(this, &SFoliageEdit::SetPaintDensity)
+						]
+					]
 
-					+ SWrapBox::Slot()
-					.Padding(FMargin(0.f, 0.f, 6.f, 0.f))
+					// Erase Density
+					+ SVerticalBox::Slot()
+					.AutoHeight()
 					[
-						SNew(SVerticalBox)
-						
-						+ SVerticalBox::Slot()
-						.Padding(FMargin(0.f, 5.f, 0.f, 0.f))
+						SNew(SHorizontalBox)
+						.ToolTipText(LOCTEXT("EraseDensity_Tooltip", "The density of foliage to leave behind when erasing with the Shift key held. 0 will remove all foliage."))
+						.Visibility(this, &SFoliageEdit::GetVisibility_EraseDensity)
+
+						+ SHorizontalBox::Slot()
+						.Padding(StandardLeftPadding)
+						.FillWidth(1.0f)
+						.VAlign(VAlign_Center)
 						[
-							SNew(SCheckBox)
-							.Visibility(this, &SFoliageEdit::GetVisibility_Filters)
-							.OnCheckStateChanged(this, &SFoliageEdit::OnCheckStateChanged_Landscape)
-							.IsChecked(this, &SFoliageEdit::GetCheckState_Landscape)
-							.ToolTipText(this, &SFoliageEdit::GetTooltipText_Landscape)
-							[
-								SNew(STextBlock)
-								.Text(LOCTEXT("Landscape", "Landscape"))
-								.Font(StandardFont)
-							]
+							SNew(STextBlock)
+							.Text(LOCTEXT("EraseDensity", "Erase Density"))
+							.Font(StandardFont)
 						]
-						
-						+ SVerticalBox::Slot()
-						.Padding(FMargin(0.f, 5.f, 0.f, 0.f))
+						+ SHorizontalBox::Slot()
+						.Padding(StandardRightPadding)
+						.FillWidth(2.0f)
+						.MaxWidth(100.f)
+						.VAlign(VAlign_Center)
 						[
-							SNew(SCheckBox)
-							.Visibility(this, &SFoliageEdit::GetVisibility_Filters)
-							.OnCheckStateChanged(this, &SFoliageEdit::OnCheckStateChanged_BSP)
-							.IsChecked(this, &SFoliageEdit::GetCheckState_BSP)
-							.ToolTipText(this, &SFoliageEdit::GetTooltipText_BSP)
-							[
-								SNew(STextBlock)
-								.Text(LOCTEXT("BSP", "BSP"))
-								.Font(StandardFont)
-							]
+							SNew(SNumericEntryBox<float>)
+							.Font(StandardFont)
+							.AllowSpin(true)
+							.MinValue(0.0f)
+							.MaxValue(2.0f)
+							.MaxSliderValue(1.0f)
+							.Value(this, &SFoliageEdit::GetEraseDensity)
+							.OnValueChanged(this, &SFoliageEdit::SetEraseDensity)
 						]
 					]
-					
-					+ SWrapBox::Slot()
-					.Padding(FMargin(0.f, 0.f, 6.f, 0.f))
+
+					// Filters
+					+ SVerticalBox::Slot()
+					.AutoHeight()
 					[
-						SNew(SVerticalBox)
+						SNew(SHorizontalBox)
+						.Visibility(this, &SFoliageEdit::GetVisibility_Filters)
 
-						+ SVerticalBox::Slot()
-						.Padding(FMargin(0.f, 5.f, 0.f, 0.f))
+						+ SHorizontalBox::Slot()
+						.VAlign(VAlign_Center)
+						.Padding(StandardPadding)
 						[
-							SNew(SCheckBox)
-							.Visibility(this, &SFoliageEdit::GetVisibility_Filters)
-							.OnCheckStateChanged(this, &SFoliageEdit::OnCheckStateChanged_StaticMesh)
-							.IsChecked(this, &SFoliageEdit::GetCheckState_StaticMesh)
-							.ToolTipText(this, &SFoliageEdit::GetTooltipText_StaticMesh)
+							SNew(SWrapBox)
+							.UseAllottedWidth(true)
+
+							+ SWrapBox::Slot()
+							.Padding(FMargin(0.f, 0.f, 6.f, 0.f))
 							[
-								SNew(STextBlock)
-								.Text(LOCTEXT("StaticMeshes", "Static Meshes"))
-								.Font(StandardFont)
+								SNew(SVerticalBox)
+
+								+ SVerticalBox::Slot()
+								.Padding(FMargin(0.f, 5.f, 0.f, 0.f))
+								[
+									SNew(SCheckBox)
+									.Visibility(this, &SFoliageEdit::GetVisibility_Filters)
+									.OnCheckStateChanged(this, &SFoliageEdit::OnCheckStateChanged_Landscape)
+									.IsChecked(this, &SFoliageEdit::GetCheckState_Landscape)
+									.ToolTipText(this, &SFoliageEdit::GetTooltipText_Landscape)
+									[
+										SNew(STextBlock)
+										.Text(LOCTEXT("Landscape", "Landscape"))
+										.Font(StandardFont)
+									]
+								]
+
+								+ SVerticalBox::Slot()
+								.Padding(FMargin(0.f, 5.f, 0.f, 0.f))
+								[
+									SNew(SCheckBox)
+									.Visibility(this, &SFoliageEdit::GetVisibility_Filters)
+									.OnCheckStateChanged(this, &SFoliageEdit::OnCheckStateChanged_BSP)
+									.IsChecked(this, &SFoliageEdit::GetCheckState_BSP)
+									.ToolTipText(this, &SFoliageEdit::GetTooltipText_BSP)
+									[
+										SNew(STextBlock)
+										.Text(LOCTEXT("BSP", "BSP"))
+										.Font(StandardFont)
+									]
+								]
 							]
-						]
 
-						+ SVerticalBox::Slot()
-						.Padding(FMargin(0.f, 5.f, 0.f, 0.f))
-						[
-							SNew(SCheckBox)
-							.Visibility(this, &SFoliageEdit::GetVisibility_Filters)
-							.OnCheckStateChanged(this, &SFoliageEdit::OnCheckStateChanged_Translucent)
-							.IsChecked(this, &SFoliageEdit::GetCheckState_Translucent)
-							.ToolTipText(this, &SFoliageEdit::GetTooltipText_Translucent)
+							+ SWrapBox::Slot()
+							.Padding(FMargin(0.f, 0.f, 6.f, 0.f))
 							[
-								SNew(STextBlock)
-								.Text(LOCTEXT("Translucent", "Translucent"))
-								.Font(StandardFont)
+								SNew(SVerticalBox)
+
+								+ SVerticalBox::Slot()
+								.Padding(FMargin(0.f, 5.f, 0.f, 0.f))
+								[
+									SNew(SCheckBox)
+									.Visibility(this, &SFoliageEdit::GetVisibility_Filters)
+									.OnCheckStateChanged(this, &SFoliageEdit::OnCheckStateChanged_StaticMesh)
+									.IsChecked(this, &SFoliageEdit::GetCheckState_StaticMesh)
+									.ToolTipText(this, &SFoliageEdit::GetTooltipText_StaticMesh)
+									[
+										SNew(STextBlock)
+										.Text(LOCTEXT("StaticMeshes", "Static Meshes"))
+										.Font(StandardFont)
+									]
+								]
+
+								+ SVerticalBox::Slot()
+								.Padding(FMargin(0.f, 5.f, 0.f, 0.f))
+								[
+									SNew(SCheckBox)
+									.Visibility(this, &SFoliageEdit::GetVisibility_Filters)
+									.OnCheckStateChanged(this, &SFoliageEdit::OnCheckStateChanged_Translucent)
+									.IsChecked(this, &SFoliageEdit::GetCheckState_Translucent)
+									.ToolTipText(this, &SFoliageEdit::GetTooltipText_Translucent)
+									[
+										SNew(STextBlock)
+										.Text(LOCTEXT("Translucent", "Translucent"))
+										.Font(StandardFont)
+									]
+								]
 							]
 						]
 					]
-				]
-			]
 
-			+ SVerticalBox::Slot()
-			.AutoHeight()
-			.Padding(StandardPadding)
-			[
-				SNew(SSeparator)
-				.Orientation(Orient_Horizontal)
-			]
-
-			+ SVerticalBox::Slot()
-			.Padding(StandardPadding)
-			.AutoHeight()
-			.HAlign(HAlign_Left)
-			[
-				SAssignNew(AddFoliageTypeCombo, SComboButton)
-				.ForegroundColor(FLinearColor::White)
-				.ButtonStyle(FEditorStyle::Get(), "FlatButton.Success")
-				.OnGetMenuContent(this, &SFoliageEdit::OnGetAddFoliageMeshAssetPicker)
-				.ContentPadding(FMargin(0.f))
-				.ButtonContent()
-				[
-					SNew(SHorizontalBox)
-					+ SHorizontalBox::Slot()
-					.VAlign(VAlign_Center)
-					.AutoWidth()
-					.Padding(1.f)
+					+ SVerticalBox::Slot()
+					.AutoHeight()
+					.HAlign(HAlign_Fill)
+					.Padding(StandardPadding)
 					[
-						SNew(STextBlock)
-						.TextStyle(FEditorStyle::Get(), "FoliageEditMode.AddFoliageType.Text")
-						.Font(FEditorStyle::Get().GetFontStyle("FontAwesome.9"))
-						.Text(FText::FromString(FString(TEXT("\xf067"))) /*fa-plus*/)
-					]
-					+ SHorizontalBox::Slot()
-					.VAlign(VAlign_Center)
-					.Padding(1.f)
-					[
-						SNew(STextBlock)
-						.Text(LOCTEXT("AddFoliageTypeButtonLabel", "Add Foliage Type"))
-						.TextStyle(FEditorStyle::Get(), "FoliageEditMode.AddFoliageType.Text")
-					]
-				]
-			]
+						SNew(SWrapBox)
+						.UseAllottedWidth(true)
+						.Visibility(this, &SFoliageEdit::GetVisibility_SelectionOptions)
 
-			+ SVerticalBox::Slot()
-			.AutoHeight()
-			.Padding(StandardPadding)
-			[
-				SNew(SBox)
-				.Visibility(this, &SFoliageEdit::GetVisibility_EmptyList)
-				.Padding(FMargin(15, 0))
-				.MinDesiredHeight(30)
-				[
-					SNew(SScaleBox)
-					.Stretch(EStretch::ScaleToFit)
-					[
-						SNew(STextBlock)
-						.Text(LOCTEXT("Foliage_DropStatic", "+ Drop Foliage Here"))
-						.ToolTipText(LOCTEXT("Foliage_DropStatic_ToolTip", "Drag static meshes/foliage settings from the content browser into this area."))
-					]
-				]
-			]
+						// Select all instances
+						+ SWrapBox::Slot()
+						.Padding(FMargin(0.f, 0.f, 5.f, 0.f))
+						[
+							SNew(SButton)
+							.OnClicked(this, &SFoliageEdit::OnSelectAllInstances)
+							.Text(LOCTEXT("SelectAllInstances", "Select All"))
+							.ToolTipText(LOCTEXT("SelectAllInstances_Tooltip", "Selects all foliage instances"))
+						]
 
-			+ SVerticalBox::Slot()
-			.FillHeight(1.0f)
-			[
-				SNew(SSplitter)
-				.Orientation(Orient_Vertical)
-				.Style(FEditorStyle::Get(), "DetailsView.Splitter")
+						// Select all invalid instances
+						+ SWrapBox::Slot()
+						.Padding(FMargin(0.f, 0.f, 5.f, 0.f))
+						[
+							SNew(SButton)
+							.OnClicked(this, &SFoliageEdit::OnSelectInvalidInstances)
+							.Text(LOCTEXT("SelectInvalidInstances", "Select Invalid"))
+							.ToolTipText(LOCTEXT("SelectInvalidInstances_Tooltip", "Selects all foliage instances that are not placed in a valid location"))
+						]
 
-				+SSplitter::Slot()
-				.Value(0.3f)
-				[
-					SAssignNew(MeshTreeWidget, SFoliageMeshTree)
-					.TreeItemsSource(&FoliageEditMode->GetFoliageMeshList())
-					.SelectionMode(ESelectionMode::Multi)
-					.OnGenerateRow(this, &SFoliageEdit::MeshTreeGenerateRow)
-					.OnGetChildren( this, &SFoliageEdit::MeshTreeGetChildren)
-					.OnContextMenuOpening(this, &SFoliageEdit::ConstructFoliageMeshContextMenu)
-					.OnSelectionChanged(this, &SFoliageEdit::MeshTreeOnSelectionChanged)
-					.HeaderRow
-					(
-						// toggle mesh for painting checkbox
-						SAssignNew(MeshTreeHeaderRowWidget, SHeaderRow)
-						+ SHeaderRow::Column(FoliageMeshColumns::ColumnID_ToggleMesh)
+						// Deselect all
+						+ SWrapBox::Slot()
+						.Padding(FMargin(0.f, 0.f, 5.f, 0.f))
 						[
-							SNew(SCheckBox)
-							.IsChecked(this, &SFoliageEdit::GetState_AllMeshes)
-							.OnCheckStateChanged(this, &SFoliageEdit::OnCheckStateChanged_AllMeshes)
+							SNew(SButton)
+							.OnClicked(this, &SFoliageEdit::OnDeselectAllInstances)
+							.Text(LOCTEXT("DeselectAllInstances", "Deselect All"))
+							.ToolTipText(LOCTEXT("DeselectAllInstances_Tooltip", "Deselects all foliage instances"))
 						]
-						.DefaultLabel(BlankText)
-						.HeaderContentPadding(FMargin(0, 1, 0, 1))
-						.HAlignHeader(HAlign_Center)
-						.HAlignCell(HAlign_Center)
-						.FixedWidth(24)
-
-						// mesh text label
-						+ SHeaderRow::Column(FoliageMeshColumns::ColumnID_MeshLabel)
-						.HeaderContentPadding(FMargin(10, 1, 0, 1))
-						.SortMode(this, &SFoliageEdit::GetMeshColumnSortMode)
-						.OnSort(this, &SFoliageEdit::OnMeshesColumnSortModeChanged)
-						.DefaultLabel(this, &SFoliageEdit::GetMeshesHeaderText)
-
-						// instance count
-						+ SHeaderRow::Column(FoliageMeshColumns::ColumnID_InstanceCount)
-						.HeaderContentPadding(FMargin(10, 1, 0, 1))
-						.DefaultLabel(LOCTEXT("InstanceCount", "Count"))
-						.DefaultTooltip(this, &SFoliageEdit::GetTotalInstanceCountTooltipText)
-						.FixedWidth(60.f)
-
-						// save Foliage asset
-						+SHeaderRow::Column(FoliageMeshColumns::ColumnID_Save).FixedWidth(24.0f)
-						.DefaultLabel(BlankText)
-					)
-				]
-					
-				+SSplitter::Slot()
-				[
-					MeshDetailsWidget.ToSharedRef()
+					]
 				]
 			]
 		]
 
-		// Foliage Mesh Drop Zone
-		+ SOverlay::Slot()
-		.HAlign(HAlign_Fill)
+		+ SVerticalBox::Slot()
+		.FillHeight(1.f)
 		.VAlign(VAlign_Fill)
+		.Padding(0.f, 5.f, 0.f, 0.f)
 		[
-			SNew(SFoliageDragDropHandler)
-			.Visibility(this, &SFoliageEdit::GetVisibility_FoliageDropTarget)
-			.FoliageEditPtr(SharedThis(this))
+			SNew(SOverlay)
+
+			+ SOverlay::Slot()
+			[
+				SAssignNew(FoliagePalette, SFoliagePalette)
+				.FoliageEdMode(FoliageEditMode)
+			]
+			
+			// Foliage Mesh Drop Zone
+			+ SOverlay::Slot()
+			.HAlign(HAlign_Fill)
+			.VAlign(VAlign_Fill)
 			[
-				SNew(SScaleBox)
-				.Stretch(EStretch::ScaleToFit)
+				SNew(SFoliageDragDropHandler)
+				.Visibility(this, &SFoliageEdit::GetVisibility_FoliageDropTarget)
+				.FoliageEditPtr(SharedThis(this))
 				[
-					SNew(STextBlock)
-					.Text(LOCTEXT("Foliage_AddFoliageMesh", "+ Foliage Mesh"))
-					.ShadowOffset(FVector2D(1.f, 1.f))
+					SNew(SScaleBox)
+					.Stretch(EStretch::ScaleToFit)
+					[
+						SNew(STextBlock)
+						.Text(LOCTEXT("Foliage_AddFoliageMesh", "+ Foliage Mesh"))
+						.ShadowOffset(FVector2D(1.f, 1.f))
+					]
 				]
 			]
 		]
@@ -486,7 +428,7 @@ END_SLATE_FUNCTION_BUILD_OPTIMIZATION
 
 void SFoliageEdit::RefreshFullList()
 {
-	MeshTreeWidget->RequestTreeRefresh();
+	FoliagePalette->Refresh();
 }
 
 void SFoliageEdit::ClearAllToolSelection()
@@ -531,42 +473,13 @@ void SFoliageEdit::BindCommands()
 		FExecuteAction::CreateSP(this, &SFoliageEdit::OnSetPaintFill),
 		FCanExecuteAction(),
 		FIsActionChecked::CreateSP(this, &SFoliageEdit::IsPaintFillTool));
-
-		//
-	UICommandList->MapAction(
-		Commands.RemoveFoliageType,
-		FExecuteAction::CreateSP(this, &SFoliageEdit::OnRemoveFoliageType),
-		FCanExecuteAction(),
-		FIsActionChecked());
-
-	UICommandList->MapAction(
-		Commands.ShowFoliageTypeInCB,
-		FExecuteAction::CreateSP(this, &SFoliageEdit::OnShowFoliageTypeInCB),
-		FCanExecuteAction(),
-		FIsActionChecked());
-
-	UICommandList->MapAction(
-		Commands.SelectAllInstances,
-		FExecuteAction::CreateSP(this, &SFoliageEdit::OnSelectAllInstances),
-		FCanExecuteAction(),
-		FIsActionChecked());
-
-	UICommandList->MapAction(
-		Commands.DeselectAllInstances,
-		FExecuteAction::CreateSP(this, &SFoliageEdit::OnDeselectAllInstances),
-		FCanExecuteAction(),
-		FIsActionChecked());
-
-	UICommandList->MapAction(
-		Commands.SelectInvalidInstances,
-		FExecuteAction::CreateSP(this, &SFoliageEdit::OnSelectInvalidInstances),
-		FCanExecuteAction(),
-		FIsActionChecked());
 }
 
 TSharedRef<SWidget> SFoliageEdit::BuildToolBar()
 {
-	FToolBarBuilder Toolbar(UICommandList, FMultiBoxCustomization::None);
+	FToolBarBuilder Toolbar(UICommandList, FMultiBoxCustomization::None, nullptr, Orient_Vertical);
+	Toolbar.SetLabelVisibility(EVisibility::Collapsed);
+	Toolbar.SetStyle(&FEditorStyle::Get(), "FoliageEditToolbar");
 	{
 		Toolbar.AddToolBarButton(FFoliageEditCommands::Get().SetPaint);
 		Toolbar.AddToolBarButton(FFoliageEditCommands::Get().SetReapplySettings);
@@ -575,13 +488,10 @@ TSharedRef<SWidget> SFoliageEdit::BuildToolBar()
 		Toolbar.AddToolBarButton(FFoliageEditCommands::Get().SetPaintBucket);
 	}
 
-	IIntroTutorials& IntroTutorials = FModuleManager::LoadModuleChecked<IIntroTutorials>(TEXT("IntroTutorials"));
-
 	return
 		SNew(SHorizontalBox)
 
 		+ SHorizontalBox::Slot()
-		.Padding(4, 0)
 		[
 			SNew(SOverlay)
 			+ SOverlay::Slot()
@@ -595,25 +505,11 @@ TSharedRef<SWidget> SFoliageEdit::BuildToolBar()
 					Toolbar.MakeWidget()
 				]
 			]
-
-			// Tutorial link
-			+ SOverlay::Slot()
-			.HAlign(HAlign_Right)
-			.VAlign(VAlign_Bottom)
-			.Padding(4)
-			[
-				IntroTutorials.CreateTutorialsWidget(TEXT("FoliageMode"))
-			]
 		];
 }
 
 void SFoliageEdit::AddFoliageType(const FAssetData& AssetData)
 {
-	if (AddFoliageTypeCombo.IsValid())
-	{
-		AddFoliageTypeCombo->SetIsOpen(false);
-	}
-
 	GWarn->BeginSlowTask(LOCTEXT("AddFoliageType_LoadPackage", "Fully Loading Package for Add"), true, false);
 	UObject* Asset = AssetData.GetAsset();
 	GWarn->EndSlowTask();
@@ -621,7 +517,7 @@ void SFoliageEdit::AddFoliageType(const FAssetData& AssetData)
 	UFoliageType* FoliageType = FoliageEditMode->AddFoliageAsset(Asset);
 	if (FoliageType)
 	{
-		MeshTreeWidget->RequestTreeRefresh();
+		FoliagePalette->Refresh();
 	}
 }
 
@@ -695,6 +591,33 @@ bool SFoliageEdit::IsPaintFillTool() const
 	return FoliageEditMode->UISettings.GetPaintBucketToolSelected();
 }
 
+FText SFoliageEdit::GetActiveToolName() const
+{
+	FText OutText;
+	if (IsPaintTool())
+	{
+		OutText = LOCTEXT("FoliageToolName_Paint", "Paint");
+	}
+	else if (IsReapplySettingsTool())
+	{
+		OutText = LOCTEXT("FoliageToolName_Reapply", "Reapply");
+	}
+	else if (IsSelectTool())
+	{
+		OutText = LOCTEXT("FoliageToolName_Select", "Select");
+	}
+	else if (IsLassoSelectTool())
+	{
+		OutText = LOCTEXT("FoliageToolName_LassoSelect", "Lasso Select");
+	}
+	else if (IsPaintFillTool())
+	{
+		OutText = LOCTEXT("FoliageToolName_Fill", "Fill");
+	}
+
+	return OutText;
+}
+
 void SFoliageEdit::SetRadius(float InRadius)
 {
 	FoliageEditMode->UISettings.SetRadius(InRadius);
@@ -726,33 +649,9 @@ TOptional<float> SFoliageEdit::GetEraseDensity() const
 	return FoliageEditMode->UISettings.GetUnpaintDensity();
 }
 
-EVisibility SFoliageEdit::GetAddFoliageTypeButtonTextVisibility() const
+EVisibility SFoliageEdit::GetVisibility_SelectionOptions() const
 {
-	return MeshTreeHeaderRowWidget.IsValid() && MeshTreeHeaderRowWidget->IsHovered() ? EVisibility::SelfHitTestInvisible : EVisibility::Collapsed;
-}
-
-TSharedRef<SWidget> SFoliageEdit::OnGetAddFoliageMeshAssetPicker()
-{
-	TArray<const UClass*> ClassFilters;
-	ClassFilters.Add(UFoliageType_InstancedStaticMesh::StaticClass());
-
-	return PropertyCustomizationHelpers::MakeAssetPickerWithMenu(FAssetData(),
-																false,
-																ClassFilters,
-																PropertyCustomizationHelpers::GetNewAssetFactoriesForClasses(ClassFilters),
-																FOnShouldFilterAsset(),
-																FOnAssetSelected::CreateSP(this, &SFoliageEdit::AddFoliageType),
-																FSimpleDelegate());
-}
-
-EVisibility SFoliageEdit::GetVisibility_EmptyList() const
-{
-	if (FoliageEditMode->GetFoliageMeshList().Num() > 0)
-	{
-		return EVisibility::Collapsed;
-	}
-
-	return EVisibility::HitTestInvisible;
+	return IsSelectTool() || IsLassoSelectTool() ? EVisibility::SelfHitTestInvisible : EVisibility::Collapsed;
 }
 
 EVisibility SFoliageEdit::GetVisibility_FoliageDropTarget() const
@@ -772,91 +671,79 @@ EVisibility SFoliageEdit::GetVisibility_FoliageDropTarget() const
 	return EVisibility::Hidden;
 }
 
-EVisibility SFoliageEdit::GetVisibility_NonEmptyList() const
+FReply SFoliageEdit::OnSelectAllInstances()
 {
-	if (FoliageEditMode->GetFoliageMeshList().Num() == 0)
+	for (FFoliageMeshUIInfoPtr& TypeInfo : FoliageEditMode->GetFoliageMeshList())
 	{
-		return EVisibility::Collapsed;
+		if (TypeInfo->InstanceCountCurrentLevel > 0)
+		{
+			UFoliageType* FoliageType = TypeInfo->Settings;
+			FoliageEditMode->SelectInstances(FoliageType, true);
+		}
 	}
 
-	return EVisibility::HitTestInvisible;
+	return FReply::Handled();
 }
 
-FReply SFoliageEdit::OnDrop_ListView(const FGeometry& MyGeometry, const FDragDropEvent& DragDropEvent)
+FReply SFoliageEdit::OnSelectInvalidInstances()
 {
-	TArray<FAssetData> DroppedAssetData = AssetUtil::ExtractAssetDataFromDrag(DragDropEvent);
+	for (FFoliageMeshUIInfoPtr& TypeInfo : FoliageEditMode->GetFoliageMeshList())
 	{
-		// Treat the entire drop as a transaction (in case multiples types are being added)
-		const FScopedTransaction Transaction(NSLOCTEXT("UnrealEd", "FoliageMode_DragDropTypesTransaction", "Drag-drop Foliage Type(s)"));
-
-		for (int32 AssetIdx = 0; AssetIdx < DroppedAssetData.Num(); ++AssetIdx)
-		{	
-			AddFoliageType(DroppedAssetData[AssetIdx]);
+		if (TypeInfo->InstanceCountCurrentLevel > 0)
+		{
+			const UFoliageType* FoliageType = TypeInfo->Settings;
+			FoliageEditMode->SelectInstances(FoliageType, false);
+			FoliageEditMode->SelectInvalidInstances(FoliageType);
 		}
 	}
 
 	return FReply::Handled();
 }
-TSharedPtr<SWidget> SFoliageEdit::ConstructFoliageMeshContextMenu() const
-{
-	const FFoliageEditCommands& Commands = FFoliageEditCommands::Get();
-	FMenuBuilder MenuBuilder(true, UICommandList);
 
-	if (MeshTreeWidget->GetSelectedItems().Num() > 0)
+FReply SFoliageEdit::OnDeselectAllInstances()
+{
+	for (FFoliageMeshUIInfoPtr& TypeInfo : FoliageEditMode->GetFoliageMeshList())
 	{
-		MenuBuilder.AddMenuEntry(Commands.RemoveFoliageType);
-
-		MenuBuilder.AddSubMenu(
-			LOCTEXT( "ReplaceFoliageType", "Replace Foliage Type" ),
-			LOCTEXT ("ReplaceFoliageType_ToolTip", "Replaces selected foliage type with another foliage type asset"),
-			FNewMenuDelegate::CreateSP(this, &SFoliageEdit::FillReplaceFoliageTypeSubmenu));
-
-		MenuBuilder.AddMenuEntry(Commands.ShowFoliageTypeInCB);
-
-		if (IsSelectTool() || IsLassoSelectTool())
+		if (TypeInfo->InstanceCountCurrentLevel > 0)
 		{
-			MenuBuilder.AddMenuEntry(Commands.SelectAllInstances);
-			MenuBuilder.AddMenuEntry(Commands.DeselectAllInstances);
-			MenuBuilder.AddMenuEntry(Commands.SelectInvalidInstances);
+			UFoliageType* FoliageType = TypeInfo->Settings;
+			FoliageEditMode->SelectInstances(FoliageType, false);
 		}
 	}
 
-	return MenuBuilder.MakeWidget();
+	return FReply::Handled();
 }
 
-FText SFoliageEdit::GetFilterText() const
+FReply SFoliageEdit::OnDrop_ListView(const FGeometry& MyGeometry, const FDragDropEvent& DragDropEvent)
 {
-	FText TooltipText;
-	if (FoliageEditMode->UISettings.GetPaintToolSelected() || FoliageEditMode->UISettings.GetPaintBucketToolSelected())
-	{
-		TooltipText = LOCTEXT("PlacementFilter", "Placement Filter");
-	}
-	else if (FoliageEditMode->UISettings.GetReapplyToolSelected())
-	{
-		TooltipText = LOCTEXT("ReapplyFilter", "Reapply Filter");
-	}
-	else if (FoliageEditMode->UISettings.GetLassoSelectToolSelected())
+	TArray<FAssetData> DroppedAssetData = AssetUtil::ExtractAssetDataFromDrag(DragDropEvent);
 	{
-		TooltipText = LOCTEXT("SelectionFilter", "Selection Filter");
+		// Treat the entire drop as a transaction (in case multiples types are being added)
+		const FScopedTransaction Transaction(NSLOCTEXT("UnrealEd", "FoliageMode_DragDropTypesTransaction", "Drag-drop Foliage Type(s)"));
+
+		for (int32 AssetIdx = 0; AssetIdx < DroppedAssetData.Num(); ++AssetIdx)
+		{	
+			AddFoliageType(DroppedAssetData[AssetIdx]);
+		}
 	}
 
-	return TooltipText;
+	return FReply::Handled();
 }
 
-FText SFoliageEdit::GetFilterTooltipText() const
+FText SFoliageEdit::GetFilterText() const
 {
 	FText TooltipText;
-	if (FoliageEditMode->UISettings.GetPaintToolSelected() || FoliageEditMode->UISettings.GetPaintBucketToolSelected())
+	if (IsPaintTool() || IsPaintFillTool())
 	{
-		TooltipText = LOCTEXT("PlacementFilter_Tooltip", "Limits the types of geometry upon which foliage instances can be placed.");
+		TooltipText = LOCTEXT("PlacementFilter", "Placement Filter");
 	}
-	else if (FoliageEditMode->UISettings.GetReapplyToolSelected())
+	else if (IsReapplySettingsTool())
 	{
-		TooltipText = LOCTEXT("ReapplyFilter_Tooltip", "Limits the types of geometry upon which foliage instances can be reapplied.");
+		TooltipText = LOCTEXT("ReapplyFilter", "Reapply Filter");
 	}
-	else if (FoliageEditMode->UISettings.GetLassoSelectToolSelected())
+	else if (IsLassoSelectTool())
 	{
-		TooltipText = LOCTEXT("SelectionFilter_Tooltip", "Limits the types of geometry upon which foliage instances can be selected.");
+		TooltipText = LOCTEXT("SelectionFilter", "Selection Filter");
 	}
 
 	return TooltipText;
@@ -875,15 +762,15 @@ ECheckBoxState SFoliageEdit::GetCheckState_Landscape() const
 FText SFoliageEdit::GetTooltipText_Landscape() const
 {
 	FText TooltipText;
-	if (FoliageEditMode->UISettings.GetPaintToolSelected() || FoliageEditMode->UISettings.GetPaintBucketToolSelected())
+	if (IsPaintTool() || IsPaintFillTool())
 	{
 		TooltipText = LOCTEXT("FilterLandscapeTooltip_Placement", "Place foliage on landscapes");
 	}
-	else if (FoliageEditMode->UISettings.GetReapplyToolSelected())
+	else if (IsReapplySettingsTool())
 	{
 		TooltipText = LOCTEXT("FilterLandscapeTooltip_Reapply", "Reapply to instances on landscapes");
 	}
-	else if (FoliageEditMode->UISettings.GetLassoSelectToolSelected())
+	else if (IsLassoSelectTool())
 	{
 		TooltipText = LOCTEXT("FilterLandscapeTooltip_Select", "Select instances on landscapes");
 	}
@@ -904,15 +791,15 @@ ECheckBoxState SFoliageEdit::GetCheckState_StaticMesh() const
 FText SFoliageEdit::GetTooltipText_StaticMesh() const
 {
 	FText TooltipText;
-	if (FoliageEditMode->UISettings.GetPaintToolSelected() || FoliageEditMode->UISettings.GetPaintBucketToolSelected())
+	if (IsPaintTool() || IsPaintFillTool())
 	{
 		TooltipText = LOCTEXT("FilterStaticMeshTooltip_Placement", "Place foliage on static meshes");
 	}
-	else if (FoliageEditMode->UISettings.GetReapplyToolSelected())
+	else if (IsReapplySettingsTool())
 	{
 		TooltipText = LOCTEXT("FilterStaticMeshTooltip_Reapply", "Reapply to instances on static meshes");
 	}
-	else if (FoliageEditMode->UISettings.GetLassoSelectToolSelected())
+	else if (IsLassoSelectTool())
 	{
 		TooltipText = LOCTEXT("FilterStaticMeshTooltip_Select", "Select instances on static meshes");
 	}
@@ -933,15 +820,15 @@ ECheckBoxState SFoliageEdit::GetCheckState_BSP() const
 FText SFoliageEdit::GetTooltipText_BSP() const
 {
 	FText TooltipText;
-	if (FoliageEditMode->UISettings.GetPaintToolSelected() || FoliageEditMode->UISettings.GetPaintBucketToolSelected())
+	if (IsPaintTool() || IsPaintFillTool())
 	{
 		TooltipText = LOCTEXT("FilterBSPTooltip_Placement", "Place foliage on BSP");
 	}
-	else if (FoliageEditMode->UISettings.GetReapplyToolSelected())
+	else if (IsReapplySettingsTool())
 	{
 		TooltipText = LOCTEXT("FilterBSPTooltip_Reapply", "Reapply to instances on BSP");
 	}
-	else if (FoliageEditMode->UISettings.GetLassoSelectToolSelected())
+	else if (IsLassoSelectTool())
 	{
 		TooltipText = LOCTEXT("FilterBSPTooltip_Select", "Select instances on BSP");
 	}
@@ -962,15 +849,15 @@ ECheckBoxState SFoliageEdit::GetCheckState_Translucent() const
 FText SFoliageEdit::GetTooltipText_Translucent() const
 {
 	FText TooltipText;
-	if (FoliageEditMode->UISettings.GetPaintToolSelected() || FoliageEditMode->UISettings.GetPaintBucketToolSelected())
+	if (IsPaintTool() || IsPaintFillTool())
 	{
 		TooltipText = LOCTEXT("FilterTranslucentTooltip_Placement", "Place foliage on translucent geometry");
 	}
-	else if (FoliageEditMode->UISettings.GetReapplyToolSelected())
+	else if (IsReapplySettingsTool())
 	{
 		TooltipText = LOCTEXT("FilterTranslucentTooltip_Reapply", "Reapply to instances on translucent geometry");
 	}
-	else if (FoliageEditMode->UISettings.GetLassoSelectToolSelected())
+	else if (IsLassoSelectTool())
 	{
 		TooltipText = LOCTEXT("FilterTranslucentTooltip_Select", "Select instances on translucent geometry");
 	}
@@ -1018,223 +905,4 @@ EVisibility SFoliageEdit::GetVisibility_Filters() const
 	return EVisibility::Visible;
 }
 
-TSharedRef<ITableRow> SFoliageEdit::MeshTreeGenerateRow(FFoliageMeshUIInfoPtr Item, const TSharedRef<STableViewBase>& OwnerTable)
-{
-	return SNew(SFoliageEditMeshRow, OwnerTable)
-		.MeshInfo(Item)
-		.FoliageEditMode(FoliageEditMode);
-}
-
-void SFoliageEdit::MeshTreeGetChildren(FFoliageMeshUIInfoPtr Item, TArray<FFoliageMeshUIInfoPtr>& OutChildren)
-{
-	//OutChildren = Item->GetChildren();
-}
-
-void SFoliageEdit::MeshTreeOnSelectionChanged(FFoliageMeshUIInfoPtr Item, ESelectInfo::Type SelectInfo)
-{
-	RefreshMeshDetailsWidget();
-}
-
-void SFoliageEdit::OnRemoveFoliageType()
-{
-	int32 NumInstances = 0;
-	TArray<UFoliageType*> FoliageTypeList;
-	
-	auto MeshUIList = MeshTreeWidget->GetSelectedItems();
-	for (FFoliageMeshUIInfoPtr& MeshUIInfo : MeshUIList)
-	{
-		NumInstances+= MeshUIInfo->InstanceCountTotal;
-		FoliageTypeList.Add(MeshUIInfo->Settings);
-	}
-	
-	bool bProceed = true;
-	if (NumInstances > 0)
-	{
-		FText Message = FText::Format(NSLOCTEXT("UnrealEd", "FoliageMode_DeleteMesh", "Are you sure you want to remove {0} instances?"), FText::AsNumber(NumInstances));
-		bProceed = (FMessageDialog::Open(EAppMsgType::YesNo, Message) == EAppReturnType::Yes);
-	}
-
-	if (bProceed)
-	{
-		FoliageEditMode->RemoveFoliageType(FoliageTypeList.GetData(), FoliageTypeList.Num());
-	}
-}
-
-void SFoliageEdit::FillReplaceFoliageTypeSubmenu(FMenuBuilder& MenuBuilder)
-{
-	FContentBrowserModule& ContentBrowserModule = FModuleManager::Get().LoadModuleChecked<FContentBrowserModule>(TEXT("ContentBrowser"));
-
-	FAssetPickerConfig AssetPickerConfig;
-	AssetPickerConfig.Filter.ClassNames.Add(UFoliageType::StaticClass()->GetFName());
-	AssetPickerConfig.Filter.bRecursiveClasses = true;
-	AssetPickerConfig.OnAssetSelected = FOnAssetSelected::CreateSP(this, &SFoliageEdit::OnReplaceFoliageTypeSelected);
-	AssetPickerConfig.InitialAssetViewType = EAssetViewType::List;
-	AssetPickerConfig.bAllowNullSelection = false;
-	
-	TSharedRef<SWidget> MenuContent = SNew(SBox)
-	.WidthOverride(384)
-	.HeightOverride(500)
-	[
-		ContentBrowserModule.Get().CreateAssetPicker(AssetPickerConfig)
-	];
-	
-	MenuBuilder.AddWidget( MenuContent, FText::GetEmpty(), true);
-}
-
-void SFoliageEdit::OnReplaceFoliageTypeSelected(const FAssetData& AssetData)
-{
-	FSlateApplication::Get().DismissAllMenus();
-	
-	auto MeshUIList = MeshTreeWidget->GetSelectedItems();
-	UFoliageType* NewSettings = Cast<UFoliageType>(AssetData.GetAsset());
-	
-	if (MeshUIList.Num() && NewSettings)
-	{
-		for (auto& MeshUI : MeshUIList)
-		{
-			UFoliageType* OldSettings = MeshUI->Settings;
-			if (OldSettings != NewSettings)
-			{
-				FoliageEditMode->ReplaceSettingsObject(OldSettings, NewSettings);
-			}
-		}
-	}
-}
-
-void SFoliageEdit::OnShowFoliageTypeInCB()
-{
-	TArray<UObject*> SelectedAssets;
-	auto MeshUIList = MeshTreeWidget->GetSelectedItems();
-	for (FFoliageMeshUIInfoPtr  MeshUIPtr : MeshUIList)
-	{
-		UFoliageType* FoliageType = MeshUIPtr->Settings;
-		if (FoliageType->IsAsset())		
-		{
-			SelectedAssets.Add(FoliageType);
-		}
-		else
-		{
-			SelectedAssets.Add(FoliageType->GetStaticMesh());
-		}
-	}
-
-	if (SelectedAssets.Num())
-	{
-		GEditor->SyncBrowserToObjects(SelectedAssets);
-	}
-}
-
-void SFoliageEdit::OnSelectAllInstances()
-{
-	TArray<UObject*> SelectedAssets;
-	auto MeshUIList = MeshTreeWidget->GetSelectedItems();
-	for (FFoliageMeshUIInfoPtr MeshUIPtr : MeshUIList)
-	{
-		UFoliageType* FoliageType = MeshUIPtr->Settings;
-		FoliageEditMode->SelectInstances(FoliageType, true);
-	}
-}
-
-void SFoliageEdit::OnDeselectAllInstances()
-{
-	TArray<UObject*> SelectedAssets;
-	auto MeshUIList = MeshTreeWidget->GetSelectedItems();
-	for (FFoliageMeshUIInfoPtr MeshUIPtr : MeshUIList)
-	{
-		UFoliageType* FoliageType = MeshUIPtr->Settings;
-		FoliageEditMode->SelectInstances(FoliageType, false);
-	}
-}
-
-void SFoliageEdit::OnSelectInvalidInstances()
-{
-	auto MeshUIList = MeshTreeWidget->GetSelectedItems();
-	for (FFoliageMeshUIInfoPtr MeshUIPtr : MeshUIList)
-	{
-		const UFoliageType* FoliageType = MeshUIPtr->Settings;
-		FoliageEditMode->SelectInvalidInstances(FoliageType);
-	}
-}
-
-ECheckBoxState SFoliageEdit::GetState_AllMeshes() const
-{
-	bool bHasChecked = false;
-	bool bHasUnchecked = false;
-
-	const auto& MeshUIList = FoliageEditMode->GetFoliageMeshList();
-	for (const auto& MeshUIPtr : MeshUIList)
-	{
-		if (MeshUIPtr->Settings->IsSelected)
-		{
-			bHasChecked = true;
-		}
-		else
-		{
-			bHasUnchecked = true;
-		}
-
-		if (bHasChecked && bHasUnchecked)
-		{
-			return ECheckBoxState::Undetermined; 
-		}
-	}
-
-	return bHasChecked ? ECheckBoxState::Checked : ECheckBoxState::Unchecked;
-}
-
-void SFoliageEdit::OnCheckStateChanged_AllMeshes(ECheckBoxState InState)
-{
-	auto& MeshUIList = FoliageEditMode->GetFoliageMeshList();
-	for (auto& MeshUIPtr : MeshUIList)
-	{
-		MeshUIPtr->Settings->IsSelected = (InState == ECheckBoxState::Checked);
-	}
-}
-
-FText SFoliageEdit::GetMeshesHeaderText() const
-{
-	int32 NumMeshes = FoliageEditMode->GetFoliageMeshList().Num();
-	return FText::Format(LOCTEXT("FoliageMeshCount", "Meshes ({0})"), FText::AsNumber(NumMeshes));
-}
-
-EColumnSortMode::Type SFoliageEdit::GetMeshColumnSortMode() const
-{
-	return FoliageEditMode->GetFoliageMeshListSortMode();
-}
-
-void SFoliageEdit::OnMeshesColumnSortModeChanged(EColumnSortPriority::Type InPriority, const FName& InColumnName, EColumnSortMode::Type InSortMode)
-{
-	FoliageEditMode->OnFoliageMeshListSortModeChanged(InSortMode);
-}
-
-FText SFoliageEdit::GetTotalInstanceCountTooltipText() const
-{
-	// Probably should cache these values, 
-	// but we call this only occasionally when tooltip is active
-	int32 InstanceCountTotal = 0;
-	int32 InstanceCountCurrentLevel = 0;
-	FoliageEditMode->CalcTotalInstanceCount(InstanceCountTotal, InstanceCountCurrentLevel);
-			
-	return FText::Format(LOCTEXT("FoliageTotalInstanceCount", "Current Level: {0} Total: {1}"), FText::AsNumber(InstanceCountCurrentLevel), FText::AsNumber(InstanceCountTotal));
-}
-
-void SFoliageEdit::HandleOnToolChanged()
-{
-	RefreshMeshDetailsWidget();
-}
-
-void SFoliageEdit::RefreshMeshDetailsWidget()
-{
-	TArray<UObject*> SelectedFoliageTypes;
-	auto SelectedFoliageInfoList = MeshTreeWidget->GetSelectedItems();
-
-	for (const auto& Info : SelectedFoliageInfoList)
-	{
-		SelectedFoliageTypes.Add(Info->Settings);
-	}
-
-	const bool bForceRefresh = true;
-	MeshDetailsWidget->SetObjects(SelectedFoliageTypes, bForceRefresh);
-}
-
 #undef LOCTEXT_NAMESPACE
diff --git a/Engine/Source/Editor/FoliageEdit/Private/SFoliageEdit.h b/Engine/Source/Editor/FoliageEdit/Private/SFoliageEdit.h
--- a/Engine/Source/Editor/FoliageEdit/Private/SFoliageEdit.h
+++ b/Engine/Source/Editor/FoliageEdit/Private/SFoliageEdit.h
@@ -2,12 +2,8 @@
 
 #pragma once
 
-// Forwards declarations
 struct FFoliageMeshUIInfo;
-class IDetailsView;
-
 typedef TSharedPtr<FFoliageMeshUIInfo> FFoliageMeshUIInfoPtr; //should match typedef in FoliageEdMode.h
-typedef STreeView<FFoliageMeshUIInfoPtr> SFoliageMeshTree;
 
 class SFoliageEdit : public SCompoundWidget
 {
@@ -28,9 +24,6 @@ public:
 	/** Gets FoliageEditMode. Used by the cluster details to notify changes */
 	class FEdModeFoliage* GetFoliageEditMode() const { return FoliageEditMode; }
 
-	/** @return the SWidget containing the context menu */
-	TSharedPtr<SWidget> ConstructFoliageMeshContextMenu() const;
-
 private:
 	/** Clears all the tools selection by setting them to false. */
 	void ClearAllToolSelection();
@@ -76,6 +69,9 @@ private:
 	/** Checks if the tool mode is Paint Bucket. */
 	bool IsPaintFillTool() const;
 
+	FText GetActiveToolName() const;
+
+private:	// BRUSH SETTINGS
 	/** Sets the brush Radius for the brush. */
 	void SetRadius(float InRadius);
 
@@ -94,18 +90,9 @@ private:
 	/** Retrieves the Erase Density for the brush. */
 	TOptional<float> GetEraseDensity() const;
 
-	/** Gets the visibility of the Add Foliage Type text in the header row button */
-	EVisibility GetAddFoliageTypeButtonTextVisibility() const;
-
-	/** Gets the asset picker for adding a foliage type. */
-	TSharedRef<SWidget> OnGetAddFoliageMeshAssetPicker();
-
 	/** Retrieves the text for the filters option */
 	FText GetFilterText() const;
 
-	/** Retrieves the tooltip text for the filters option */
-	FText GetFilterTooltipText() const;
-
 	/** Sets the filter settings for if painting will occur on Landscapes. */
 	void OnCheckStateChanged_Landscape(ECheckBoxState InState);
 
@@ -142,15 +129,6 @@ private:
 	/** Retrieves the tooltip text for the translucent filter */
 	FText GetTooltipText_Translucent() const;
 
-	/** Checks if the text in the empty list overlay should appear. If the list is has items but the the drag and drop override is true, it will return EVisibility::Visible. */
-	EVisibility GetVisibility_EmptyList() const;
-
-	/** Should the drop area be visible currently?  Happens when the user is dragging static meshes */
-	EVisibility GetVisibility_FoliageDropTarget() const;
-
-	/** Checks if the list should appear. */
-	EVisibility GetVisibility_NonEmptyList() const;
-
 	/** Checks if the radius spinbox should appear. Dependant on the current tool being used. */
 	EVisibility GetVisibility_Radius() const;
 
@@ -163,46 +141,25 @@ private:
 	/** Checks if the filters should appear. Dependant on the current tool being used. */
 	EVisibility GetVisibility_Filters() const;
 
-	/** Generates a row widget for foliage mesh item */
-	TSharedRef<ITableRow> MeshTreeGenerateRow(FFoliageMeshUIInfoPtr Item, const TSharedRef<STableViewBase>& OwnerTable);
-	
-	/** Generates a list of children items for foliage item */
-	void MeshTreeGetChildren(FFoliageMeshUIInfoPtr Item, TArray<FFoliageMeshUIInfoPtr>& OutChildren);
+	/** Checks if the text in the empty list overlay should appear. If the list is has items but the the drag and drop override is true, it will return EVisibility::Visible. */
+	EVisibility GetVisibility_EmptyList() const;
 
-	/** Handler for mesh list view selection changes  */
-	void MeshTreeOnSelectionChanged(FFoliageMeshUIInfoPtr Item, ESelectInfo::Type SelectInfo);
+	/** Should the drop area be visible currently?  Happens when the user is dragging static meshes */
+	EVisibility GetVisibility_FoliageDropTarget() const;
 
-	/** Fills 'Replace' menu command  */
-	void FillReplaceFoliageTypeSubmenu(FMenuBuilder& MenuBuilder);
-	
-	/** Handler for 'Remove' command  */
-	void OnRemoveFoliageType();
-	
-	/** Handler for 'Show in CB' command  */
-	void OnShowFoliageTypeInCB();
+	/** @return Whether selecting instances is currently possible */
+	EVisibility GetVisibility_SelectionOptions() const;
 
-	/** Handler for 'Replace' command  */
-	void OnReplaceFoliageTypeSelected(const class FAssetData& AssetData);
+private:	// SELECTION
 
 	/** Handler for 'Select All' command  */
-	void OnSelectAllInstances();
-
-	/** Handler for 'Deselect All' command  */
-	void OnDeselectAllInstances();
+	FReply OnSelectAllInstances();
 
 	/** Handler for 'Select Invalid Instances' command  */
-	void OnSelectInvalidInstances();
-	
-	/** Toggle all meshes on/off */
-	ECheckBoxState GetState_AllMeshes() const;
-	void OnCheckStateChanged_AllMeshes(ECheckBoxState InState);
+	FReply OnSelectInvalidInstances();
 
-	/** Text for foliage meshes list header */
-	FText GetMeshesHeaderText() const;
-	
-	/** Mesh list sorting support */
-	EColumnSortMode::Type GetMeshColumnSortMode() const;
-	void OnMeshesColumnSortModeChanged(EColumnSortPriority::Type InPriority, const FName& InColumnName, EColumnSortMode::Type InSortMode);
+	/** Handler for 'Deselect All' command  */
+	FReply OnDeselectAllInstances();
 
 	/** Tooltip text for 'Instance Count" column */
 	FText GetTotalInstanceCountTooltipText() const;
@@ -210,25 +167,14 @@ private:
 	/** Handler to trigger a refresh of the details view when the active tool changes */
 	void HandleOnToolChanged();
 
-	/** Refreshes the mesh details widget to match the current selection */
-	void RefreshMeshDetailsWidget();
-
-private:
-	/** Foliage mesh tree widget  */
-	TSharedPtr<SFoliageMeshTree>	MeshTreeWidget;
-	
-	/** Foliage mesh details widget  */
-	TSharedPtr<IDetailsView>		MeshDetailsWidget;
 	
+private:
+	/** Palette of available foliage types */
+	TSharedPtr<class SFoliagePalette> FoliagePalette;
+
 	/** Command list for binding functions for the toolbar. */
 	TSharedPtr<FUICommandList>		UICommandList;
 
 	/** Pointer to the foliage edit mode. */
 	FEdModeFoliage*					FoliageEditMode;
-
-	/** The Add Foliage Type combo button */
-	TSharedPtr<class SComboButton> AddFoliageTypeCombo;
-
-	/** The header row of the foliage mesh tree */
-	TSharedPtr<class SHeaderRow> MeshTreeHeaderRowWidget;
 };
diff --git a/Engine/Source/Editor/FoliageEdit/Private/SFoliageEditMeshRow.cpp b/Engine/Source/Editor/FoliageEdit/Private/SFoliageEditMeshRow.cpp
deleted file mode 100644
diff --git a/Engine/Source/Editor/FoliageEdit/Private/SFoliageEditMeshRow.h b/Engine/Source/Editor/FoliageEdit/Private/SFoliageEditMeshRow.h
deleted file mode 100644
diff --git a/Engine/Source/Editor/FoliageEdit/Private/SFoliagePalette.cpp b/Engine/Source/Editor/FoliageEdit/Private/SFoliagePalette.cpp
new file mode 100644
--- /dev/null
+++ b/Engine/Source/Editor/FoliageEdit/Private/SFoliagePalette.cpp
@@ -0,0 +1,882 @@
+// Copyright 1998-2015 Epic Games, Inc. All Rights Reserved.
+#include "UnrealEd.h"
+#include "FoliageType_InstancedStaticMesh.h"
+#include "FoliageEdMode.h"
+#include "FoliageEditActions.h"
+#include "SFoliagePalette.h"
+#include "FoliageTypePaintingCustomization.h"
+#include "FoliagePaletteItem.h"
+
+#include "LevelEditor.h"
+#include "Editor/UnrealEd/Public/AssetThumbnail.h"
+#include "Editor/ContentBrowser/Public/ContentBrowserModule.h"
+#include "Editor/PropertyEditor/Public/PropertyCustomizationHelpers.h"
+#include "Editor/UnrealEd/Public/DragAndDrop/AssetDragDropOp.h"
+#include "Editor/UnrealEd/Public/AssetSelection.h"
+#include "Editor/PropertyEditor/Public/IDetailsView.h"
+#include "Editor/PropertyEditor/Public/PropertyEditorModule.h"
+#include "Engine/StaticMesh.h"
+
+#include "SScaleBox.h"
+#include "SWidgetSwitcher.h"
+#include "SExpandableArea.h"
+
+#define LOCTEXT_NAMESPACE "FoliageEd_Mode"
+
+void SFoliagePalette::Construct(const FArguments& InArgs)
+{
+	FoliageEditMode = InArgs._FoliageEdMode;
+
+	//@todo: Save these as config options
+	bShowTooltips = true;
+	ActiveViewMode = EPaletteViewMode::Thumbnail;
+
+	FoliageEditMode->OnToolChanged.AddSP(this, &SFoliagePalette::HandleOnToolChanged);
+
+	BindCommands();
+
+	FPropertyEditorModule& PropertyModule = FModuleManager::LoadModuleChecked<FPropertyEditorModule>("PropertyEditor");
+	FDetailsViewArgs Args(false, false, false, FDetailsViewArgs::HideNameArea, true);
+	Args.bShowActorLabel = false;
+	DetailsWidget = PropertyModule.CreateDetailView(Args);
+
+	// We want to use our own customization for UFoliageType
+	DetailsWidget->RegisterInstancedCustomPropertyLayout(UFoliageType::StaticClass(), 
+		FOnGetDetailCustomizationInstance::CreateStatic(&FFoliageTypePaintingCustomization::MakeInstance, FoliageEditMode)
+		);
+
+	FMargin StandardPadding(6.f, 3.f);
+	const FText BlankText = LOCTEXT("Blank", "");
+
+	// Tile View Widget
+	static const float TileSize = 64.f;
+	SAssignNew(TileViewWidget, SFoliageTypeTileView)
+		.ListItemsSource(&FoliageEditMode->GetFoliageMeshList())
+		.SelectionMode(ESelectionMode::Multi)
+		.OnGenerateTile(this, &SFoliagePalette::GenerateTile)
+		.OnContextMenuOpening(this, &SFoliagePalette::ConstructFoliageTypeContextMenu)
+		.OnSelectionChanged(this, &SFoliagePalette::OnSelectionChanged)
+		.ItemHeight(TileSize)
+		.ItemWidth(TileSize)
+		.ItemAlignment(EListItemAlignment::LeftAligned)
+		.ClearSelectionOnClick(true)
+		.OnMouseButtonDoubleClick(this, &SFoliagePalette::OnItemDoubleClicked);
+
+	// Tree View Widget
+	SAssignNew(TreeViewWidget, SFoliageTypeTreeView)
+	.TreeItemsSource(&FoliageEditMode->GetFoliageMeshList())
+	.SelectionMode(ESelectionMode::Multi)
+	.OnGenerateRow(this, &SFoliagePalette::TreeViewGenerateRow)
+	.OnGetChildren(this, &SFoliagePalette::TreeViewGetChildren)
+	.OnContextMenuOpening(this, &SFoliagePalette::ConstructFoliageTypeContextMenu)
+	.OnSelectionChanged(this, &SFoliagePalette::OnSelectionChanged)
+	.OnMouseButtonDoubleClick(this, &SFoliagePalette::OnItemDoubleClicked)
+	.HeaderRow
+	(
+		// Toggle Active
+		SAssignNew(TreeViewHeaderRow, SHeaderRow)
+		+ SHeaderRow::Column(FoliagePaletteTreeColumns::ColumnID_ToggleActive)
+		[
+			SNew(SCheckBox)
+			.IsChecked(this, &SFoliagePalette::GetState_AllMeshes)
+			.OnCheckStateChanged(this, &SFoliagePalette::OnCheckStateChanged_AllMeshes)
+		]
+		.DefaultLabel(BlankText)
+		.HeaderContentPadding(FMargin(0, 1, 0, 1))
+		.HAlignHeader(HAlign_Center)
+		.HAlignCell(HAlign_Center)
+		.FixedWidth(24)
+
+		// Type
+		+ SHeaderRow::Column(FoliagePaletteTreeColumns::ColumnID_Type)
+		.HeaderContentPadding(FMargin(10, 1, 0, 1))
+		.SortMode(this, &SFoliagePalette::GetMeshColumnSortMode)
+		.OnSort(this, &SFoliagePalette::OnMeshesColumnSortModeChanged)
+		.DefaultLabel(this, &SFoliagePalette::GetMeshesHeaderText)
+
+		// Instance Count
+		+ SHeaderRow::Column(FoliagePaletteTreeColumns::ColumnID_InstanceCount)
+		.HeaderContentPadding(FMargin(10, 1, 0, 1))
+		.DefaultLabel(LOCTEXT("InstanceCount", "Count"))
+		.DefaultTooltip(this, &SFoliagePalette::GetTotalInstanceCountTooltipText)
+		.FixedWidth(60.f)
+
+		// Save Asset
+		+ SHeaderRow::Column(FoliagePaletteTreeColumns::ColumnID_Save)
+		.FixedWidth(24.0f)
+		.DefaultLabel(BlankText)
+	);
+
+	ChildSlot
+	[
+		SNew(SVerticalBox)
+
+		+ SVerticalBox::Slot()
+		.AutoHeight()
+		.HAlign(HAlign_Fill)
+		[
+			// Top bar
+			SNew(SBorder)
+			.BorderImage(FEditorStyle::GetBrush("DetailsView.CategoryTop"))
+			.Padding(FMargin(6.f, 2.f))
+			.BorderBackgroundColor(FLinearColor(.6f, .6f, .6f, 1.0f))
+			[
+				SNew(SHorizontalBox)
+
+				+ SHorizontalBox::Slot()
+				.HAlign(HAlign_Left)
+				.AutoWidth()
+				[
+					// +Add Foliage Type button
+					SAssignNew(AddFoliageTypeCombo, SComboButton)
+					.ForegroundColor(FLinearColor::White)
+					.ButtonStyle(FEditorStyle::Get(), "FlatButton.Success")
+					.OnGetMenuContent(this, &SFoliagePalette::GetAddFoliageTypePicker)
+					.ContentPadding(FMargin(1.f))
+					.ButtonContent()
+					[
+						SNew(SHorizontalBox)
+						+ SHorizontalBox::Slot()
+						.VAlign(VAlign_Center)
+						.AutoWidth()
+						.Padding(1.f)
+						[
+							SNew(STextBlock)
+							.TextStyle(FEditorStyle::Get(), "FoliageEditMode.AddFoliageType.Text")
+							.Font(FEditorStyle::Get().GetFontStyle("FontAwesome.9"))
+							.Text(FText::FromString(FString(TEXT("\xf067"))) /*fa-plus*/)
+						]
+						+ SHorizontalBox::Slot()
+						.VAlign(VAlign_Center)
+						.Padding(1.f)
+						[
+							SNew(STextBlock)
+							.Text(LOCTEXT("AddFoliageTypeButtonLabel", "Add Foliage Type"))
+							.TextStyle(FEditorStyle::Get(), "FoliageEditMode.AddFoliageType.Text")
+						]
+					]
+				]
+
+				+ SHorizontalBox::Slot()
+				.HAlign(HAlign_Right)
+				[
+					// View Options
+					SNew( SComboButton )
+					.ContentPadding(0)
+					.ForegroundColor( FSlateColor::UseForeground() )
+					.ButtonStyle( FEditorStyle::Get(), "ToggleButton" )
+					.OnGetMenuContent(this, &SFoliagePalette::GetViewOptionsMenuContent)
+					.ButtonContent()
+					[
+						SNew(SImage)
+						.Image( FEditorStyle::GetBrush("GenericViewButton") )
+					]
+				]
+				
+			]
+		]
+
+		+ SVerticalBox::Slot()
+		[
+			SNew(SSplitter)
+			.Orientation(Orient_Vertical)
+			.Style(FEditorStyle::Get(), "FoliageEditMode.Splitter")
+
+			+ SSplitter::Slot()
+			.Value(0.6f)
+			[
+				SNew(SVerticalBox)
+
+				+ SVerticalBox::Slot()
+				.AutoHeight()
+				.Padding(StandardPadding)
+				[
+					SNew(SBox)
+					.Visibility(this, &SFoliagePalette::GetDropFoliageHintVisibility)
+					.Padding(FMargin(15, 0))
+					.MinDesiredHeight(30)
+					[
+						SNew(SScaleBox)
+						.Stretch(EStretch::ScaleToFit)
+						[
+							SNew(STextBlock)
+							.Text(LOCTEXT("Foliage_DropStatic", "+ Drop Foliage Here"))
+							.ToolTipText(LOCTEXT("Foliage_DropStatic_ToolTip", "Drag and drop foliage types or static meshes from the Content Browser to add them to the palette"))
+						]
+					]
+				]
+
+				+ SVerticalBox::Slot()
+				[
+					SAssignNew(WidgetSwitcher, SWidgetSwitcher)
+
+					+ SWidgetSwitcher::Slot()
+					[
+						// Thumbnail View
+						SNew(SScrollBorder, TileViewWidget.ToSharedRef())
+						.Content()
+						[
+							TileViewWidget.ToSharedRef()
+						]
+					]
+
+					+ SWidgetSwitcher::Slot()
+					[
+						// Tree View
+						SNew(SScrollBorder, TreeViewWidget.ToSharedRef())
+						.Style(&FEditorStyle::Get().GetWidgetStyle<FScrollBorderStyle>("FoliageEditMode.TreeView.ScrollBorder"))
+						.Content()
+						[
+							TreeViewWidget.ToSharedRef()
+						]
+					]
+				]
+				
+				+ SVerticalBox::Slot()
+				.Padding(FMargin(0.f))
+				.VAlign(VAlign_Bottom)
+				.AutoHeight()
+				[
+					SNew(SHorizontalBox)
+					
+					// Selected type name area
+					+ SHorizontalBox::Slot()
+					.Padding(FMargin(3.f))
+					.VAlign(VAlign_Bottom)
+					.AutoWidth()
+					[
+						SNew(STextBlock)
+						.Text(this, &SFoliagePalette::GetDetailsNameAreaText)
+					]
+
+					// Show/Hide Details
+					+ SHorizontalBox::Slot()
+					.HAlign(HAlign_Right)
+					[
+						SNew(SButton)
+						.ToolTipText(this, &SFoliagePalette::GetShowHideDetailsTooltipText)
+						.ForegroundColor(FSlateColor::UseForeground())
+						.ButtonStyle(FEditorStyle::Get(), "ToggleButton")
+						.OnClicked(this, &SFoliagePalette::OnShowHideDetailsClicked)
+						.ContentPadding(FMargin(2.f))
+						.Content()
+						[
+							SNew(SHorizontalBox)
+
+							+SHorizontalBox::Slot()
+							.AutoWidth()
+							.HAlign(HAlign_Center)
+							.VAlign(VAlign_Center)
+							[
+								SNew(SImage)
+								.Image(FEditorStyle::GetBrush("LevelEditor.Tabs.Details"))
+							]
+
+							+ SHorizontalBox::Slot()
+							.Padding(FMargin(3.f, 0.f))
+							.AutoWidth()
+							.HAlign(HAlign_Center)
+							.VAlign(VAlign_Center)
+							[
+								SNew(SImage)
+								.Image(this, &SFoliagePalette::GetShowHideDetailsImage)
+							]
+						]
+					]
+				]
+			]
+
+			// Details
+			+SSplitter::Slot()
+			[
+				DetailsWidget.ToSharedRef()
+			]
+		]
+	];
+}
+
+void SFoliagePalette::Refresh()
+{
+	TreeViewWidget->RequestTreeRefresh();
+	TileViewWidget->RequestListRefresh();
+}
+
+bool SFoliagePalette::AnySelectedTileHovered() const
+{
+	for (auto& TypeInfo : GetActiveViewWidget()->GetSelectedItems())
+	{
+		if (TileViewWidget->WidgetFromItem(TypeInfo)->AsWidget()->IsHovered())
+		{
+			return true;
+		}
+	}
+
+	return false;
+}
+
+void SFoliagePalette::ActivateAllSelectedTypes(bool bActivate) const
+{
+	// Apply the new check state to all of the selected types
+	for (FFoliageMeshUIInfoPtr& TypeInfo : GetActiveViewWidget()->GetSelectedItems())
+	{
+		if (TypeInfo->Settings->IsSelected != bActivate)
+		{
+			TypeInfo->Settings->Modify();
+			TypeInfo->Settings->IsSelected = bActivate;
+		}
+	}
+}
+
+void SFoliagePalette::BindCommands()
+{
+	UICommandList = MakeShareable(new FUICommandList);
+	const FFoliageEditCommands& Commands = FFoliageEditCommands::Get();
+
+	// Context menu commands
+	UICommandList->MapAction(
+		Commands.RemoveFoliageType,
+		FExecuteAction::CreateSP(this, &SFoliagePalette::OnRemoveFoliageType),
+		FCanExecuteAction());
+
+	UICommandList->MapAction(
+		Commands.ShowFoliageTypeInCB,
+		FExecuteAction::CreateSP(this, &SFoliagePalette::OnShowFoliageTypeInCB),
+		FCanExecuteAction());
+
+	UICommandList->MapAction(
+		Commands.SelectAllInstances,
+		FExecuteAction::CreateSP(this, &SFoliagePalette::OnSelectAllInstances),
+		FCanExecuteAction::CreateSP(this, &SFoliagePalette::CanSelectInstances));
+
+	UICommandList->MapAction(
+		Commands.DeselectAllInstances,
+		FExecuteAction::CreateSP(this, &SFoliagePalette::OnDeselectAllInstances),
+		FCanExecuteAction::CreateSP(this, &SFoliagePalette::CanSelectInstances));
+
+	UICommandList->MapAction(
+		Commands.SelectInvalidInstances,
+		FExecuteAction::CreateSP(this, &SFoliagePalette::OnSelectInvalidInstances),
+		FCanExecuteAction::CreateSP(this, &SFoliagePalette::CanSelectInstances));
+}
+
+void SFoliagePalette::AddFoliageType(const FAssetData& AssetData)
+{
+	if (AddFoliageTypeCombo.IsValid())
+	{
+		AddFoliageTypeCombo->SetIsOpen(false);
+	}
+
+	GWarn->BeginSlowTask(LOCTEXT("AddFoliageType_LoadPackage", "Loading Foliage Type"), true, false);
+	UObject* Asset = AssetData.GetAsset();
+	GWarn->EndSlowTask();
+
+	UFoliageType* FoliageType = FoliageEditMode->AddFoliageAsset(Asset);
+	if (FoliageType)
+	{
+		Refresh();
+	}
+}
+
+TSharedRef<SWidget> SFoliagePalette::GetAddFoliageTypePicker()
+{
+	TArray<const UClass*> ClassFilters;
+	ClassFilters.Add(UFoliageType_InstancedStaticMesh::StaticClass());
+
+	return PropertyCustomizationHelpers::MakeAssetPickerWithMenu(FAssetData(),
+		false,
+		ClassFilters,
+		PropertyCustomizationHelpers::GetNewAssetFactoriesForClasses(ClassFilters),
+		FOnShouldFilterAsset(),
+		FOnAssetSelected::CreateSP(this, &SFoliagePalette::AddFoliageType),
+		FSimpleDelegate());
+}
+
+void SFoliagePalette::HandleOnToolChanged()
+{
+	RefreshMeshDetailsWidget();
+}
+
+void SFoliagePalette::SetViewMode(EPaletteViewMode::Type NewViewMode)
+{
+	if (ActiveViewMode != NewViewMode)
+	{
+		switch (NewViewMode)
+		{
+		case EPaletteViewMode::Thumbnail:
+			// Set the tile selection to be the current tree selections
+			TileViewWidget->ClearSelection();
+			for (auto& TypeInfo : TreeViewWidget->GetSelectedItems())
+			{
+				TileViewWidget->SetItemSelection(TypeInfo, true);
+			}
+			break;
+			
+		case EPaletteViewMode::Tree:
+			// Set the tree selection to be the current tile selection
+			TreeViewWidget->ClearSelection();
+			for (auto& TypeInfo : TileViewWidget->GetSelectedItems())
+			{
+				TreeViewWidget->SetItemSelection(TypeInfo, true);
+			}
+			break;
+		}
+
+		ActiveViewMode = NewViewMode;
+		WidgetSwitcher->SetActiveWidgetIndex(ActiveViewMode);
+	}
+}
+
+bool SFoliagePalette::IsActiveViewMode(EPaletteViewMode::Type ViewMode) const
+{
+	return ActiveViewMode == ViewMode;
+}
+
+void SFoliagePalette::ToggleShowTooltips()
+{
+	bShowTooltips = !bShowTooltips;
+}
+
+bool SFoliagePalette::ShouldShowTooltips() const
+{
+	return bShowTooltips;
+}
+
+void SFoliagePalette::OnSelectionChanged(FFoliageMeshUIInfoPtr Item, ESelectInfo::Type SelectInfo)
+{
+	RefreshMeshDetailsWidget();
+}
+
+void SFoliagePalette::OnItemDoubleClicked(FFoliageMeshUIInfoPtr Item) const
+{
+	Item->Settings->IsSelected = !Item->Settings->IsSelected;
+}
+
+TSharedRef<SWidget> SFoliagePalette::GetViewOptionsMenuContent()
+{	
+	const FFoliageEditCommands& Commands = FFoliageEditCommands::Get();
+	FMenuBuilder MenuBuilder(true, UICommandList);
+
+	MenuBuilder.BeginSection("FoliagePaletteViewMode", LOCTEXT("ViewModeHeading", "Palette View Mode"));
+	{
+		MenuBuilder.AddMenuEntry(
+			LOCTEXT("ThumbnailView", "Thumbnails"),
+			LOCTEXT("ThumbnailView_ToolTip", "Display thumbnails for each foliage type in the palette."),
+			FSlateIcon(),
+			FUIAction(
+				FExecuteAction::CreateSP(this, &SFoliagePalette::SetViewMode, EPaletteViewMode::Thumbnail),
+				FCanExecuteAction(),
+				FIsActionChecked::CreateSP(this, &SFoliagePalette::IsActiveViewMode, EPaletteViewMode::Thumbnail)
+			),
+			NAME_None,
+			EUserInterfaceActionType::RadioButton
+		);
+
+		MenuBuilder.AddMenuEntry(
+			LOCTEXT("ListView", "List"),
+			LOCTEXT("ListView_ToolTip", "Display foliage types in the palette as a list."),
+			FSlateIcon(),
+			FUIAction(
+				FExecuteAction::CreateSP(this, &SFoliagePalette::SetViewMode, EPaletteViewMode::Tree),
+				FCanExecuteAction(),
+				FIsActionChecked::CreateSP(this, &SFoliagePalette::IsActiveViewMode, EPaletteViewMode::Tree)
+			),
+			NAME_None,
+			EUserInterfaceActionType::RadioButton
+		);
+	}
+	MenuBuilder.EndSection();
+
+	MenuBuilder.BeginSection("FoliagePaletteViewOptions", LOCTEXT("ViewOptionsHeading", "View Options"));
+	{
+		MenuBuilder.AddMenuEntry(
+			LOCTEXT("ShowTooltips", "Show Tooltips"),
+			LOCTEXT("ShowTooltips_ToolTip", "Whether to show tooltips when hovering over foliage types in the palette."),
+			FSlateIcon(),
+			FUIAction(
+				FExecuteAction::CreateSP(this, &SFoliagePalette::ToggleShowTooltips),
+				FCanExecuteAction(),
+				FIsActionChecked::CreateSP(this, &SFoliagePalette::ShouldShowTooltips)
+			),
+			NAME_None,
+			EUserInterfaceActionType::ToggleButton
+		);
+
+		// Tile size slider
+		// Hide inactive types
+	}
+	MenuBuilder.EndSection();
+
+	return MenuBuilder.MakeWidget();
+}
+
+TSharedPtr<SListView<FFoliageMeshUIInfoPtr>> SFoliagePalette::GetActiveViewWidget() const
+{
+	if (ActiveViewMode == EPaletteViewMode::Thumbnail)
+	{
+		return TileViewWidget;
+	}
+	else if (ActiveViewMode == EPaletteViewMode::Tree)
+	{
+		return TreeViewWidget;
+	}
+	
+	return nullptr;
+}
+
+EVisibility SFoliagePalette::GetDropFoliageHintVisibility() const
+{
+	return FoliageEditMode->GetFoliageMeshList().Num() == 0 ? EVisibility::Visible : EVisibility::Collapsed;
+}
+
+//	CONTEXT MENU
+
+TSharedPtr<SWidget> SFoliagePalette::ConstructFoliageTypeContextMenu()
+{
+	const FFoliageEditCommands& Commands = FFoliageEditCommands::Get();
+	FMenuBuilder MenuBuilder(true, UICommandList);
+
+	if (GetActiveViewWidget()->GetSelectedItems().Num() > 0)
+	{
+		if (AreAnyNonAssetTypesSelected())
+		{
+			MenuBuilder.BeginSection("StaticMeshFoliageTypeOptions", LOCTEXT("StaticMeshFoliageTypeOptionsHeader", "Static Mesh"));
+			{
+				MenuBuilder.AddMenuEntry(
+					LOCTEXT("SaveAsFoliageType", "Save As Foliage Type..."),
+					LOCTEXT("SaveAsFoliageType_ToolTip", "Creates a Foliage Type asset with these settings that can be reused in other levels."),
+					FSlateIcon(FEditorStyle::GetStyleSetName(), "Level.SaveIcon16x"),
+					FUIAction(
+						FExecuteAction::CreateSP(this, &SFoliagePalette::OnSaveSelectedAsFoliageType),
+						FCanExecuteAction::CreateSP(this, &SFoliagePalette::OnCanSaveSelectedAsFoliageType)
+					),
+					NAME_None
+				);
+			}
+			MenuBuilder.EndSection();
+		}
+
+		MenuBuilder.BeginSection("FoliageTypeOptions", LOCTEXT("FoliageTypeOptionsHeader", "Foliage Type"));
+		{
+			MenuBuilder.AddMenuEntry(Commands.RemoveFoliageType);
+
+			MenuBuilder.AddSubMenu(
+				LOCTEXT("ReplaceFoliageType", "Replace With..."),
+				LOCTEXT("ReplaceFoliageType_ToolTip", "Replaces selected foliage type with another foliage type asset"),
+				FNewMenuDelegate::CreateSP(this, &SFoliagePalette::FillReplaceFoliageTypeSubmenu));
+
+			MenuBuilder.AddMenuEntry(Commands.ShowFoliageTypeInCB);
+		}
+		MenuBuilder.EndSection();
+		
+		MenuBuilder.BeginSection("InstanceSelectionOptions", LOCTEXT("InstanceSelectionOptionsHeader", "Selection"));
+		{
+			MenuBuilder.AddMenuEntry(Commands.SelectAllInstances);
+			MenuBuilder.AddMenuEntry(Commands.DeselectAllInstances);
+			MenuBuilder.AddMenuEntry(Commands.SelectInvalidInstances);
+		}
+		MenuBuilder.EndSection();
+	}
+
+	return MenuBuilder.MakeWidget();
+}
+
+void SFoliagePalette::OnSaveSelectedAsFoliageType()
+{
+	for (FFoliageMeshUIInfoPtr& TypeInfo : GetActiveViewWidget()->GetSelectedItems())
+	{
+		UFoliageType* SavedSettings = FoliageEditMode->SaveFoliageTypeObject(TypeInfo->Settings);
+		if (SavedSettings)
+		{
+			TypeInfo->Settings = SavedSettings;
+		}
+	}
+}
+
+bool SFoliagePalette::OnCanSaveSelectedAsFoliageType() const
+{
+	for (FFoliageMeshUIInfoPtr& TypeInfo : GetActiveViewWidget()->GetSelectedItems())
+	{
+		if (TypeInfo->Settings->IsAsset())
+		{
+			// At least one selected type is an asset
+			return false;
+		}
+	}
+
+	return true;
+}
+
+bool SFoliagePalette::AreAnyNonAssetTypesSelected() const
+{
+	for (FFoliageMeshUIInfoPtr& TypeInfo : GetActiveViewWidget()->GetSelectedItems())
+	{
+		if (!TypeInfo->Settings->IsAsset())
+		{
+			// At least one selected type isn't an asset
+			return true;
+		}
+	}
+
+	return false;
+}
+
+void SFoliagePalette::FillReplaceFoliageTypeSubmenu(FMenuBuilder& MenuBuilder)
+{
+	FContentBrowserModule& ContentBrowserModule = FModuleManager::Get().LoadModuleChecked<FContentBrowserModule>(TEXT("ContentBrowser"));
+
+	FAssetPickerConfig AssetPickerConfig;
+	AssetPickerConfig.Filter.ClassNames.Add(UFoliageType::StaticClass()->GetFName());
+	AssetPickerConfig.Filter.bRecursiveClasses = true;
+	AssetPickerConfig.OnAssetSelected = FOnAssetSelected::CreateSP(this, &SFoliagePalette::OnReplaceFoliageTypeSelected);
+	AssetPickerConfig.InitialAssetViewType = EAssetViewType::List;
+	AssetPickerConfig.bAllowNullSelection = false;
+
+	TSharedRef<SWidget> MenuContent = SNew(SBox)
+		.WidthOverride(384)
+		.HeightOverride(500)
+		[
+			ContentBrowserModule.Get().CreateAssetPicker(AssetPickerConfig)
+		];
+
+	MenuBuilder.AddWidget(MenuContent, FText::GetEmpty(), true);
+}
+
+void SFoliagePalette::OnReplaceFoliageTypeSelected(const FAssetData& AssetData)
+{
+	FSlateApplication::Get().DismissAllMenus();
+
+	UFoliageType* NewSettings = Cast<UFoliageType>(AssetData.GetAsset());
+	if (GetActiveViewWidget()->GetSelectedItems().Num() && NewSettings)
+	{
+		for (FFoliageMeshUIInfoPtr& TypeInfo : GetActiveViewWidget()->GetSelectedItems())
+		{
+			UFoliageType* OldSettings = TypeInfo->Settings;
+			if (OldSettings != NewSettings)
+			{
+				FoliageEditMode->ReplaceSettingsObject(OldSettings, NewSettings);
+			}
+		}
+	}
+}
+
+void SFoliagePalette::OnRemoveFoliageType()
+{
+	int32 NumInstances = 0;
+	TArray<UFoliageType*> FoliageTypeList;
+	for (FFoliageMeshUIInfoPtr& TypeInfo : GetActiveViewWidget()->GetSelectedItems())
+	{
+		NumInstances += TypeInfo->InstanceCountTotal;
+		FoliageTypeList.Add(TypeInfo->Settings);
+	}
+
+	bool bProceed = true;
+	if (NumInstances > 0)
+	{
+		FText Message = FText::Format(NSLOCTEXT("UnrealEd", "FoliageMode_DeleteMesh", "Are you sure you want to remove {0} instances?"), FText::AsNumber(NumInstances));
+		bProceed = (FMessageDialog::Open(EAppMsgType::YesNo, Message) == EAppReturnType::Yes);
+	}
+
+	if (bProceed)
+	{
+		FoliageEditMode->RemoveFoliageType(FoliageTypeList.GetData(), FoliageTypeList.Num());
+	}
+}
+
+void SFoliagePalette::OnShowFoliageTypeInCB()
+{
+	TArray<UObject*> SelectedAssets;
+	for (FFoliageMeshUIInfoPtr& TypeInfo : GetActiveViewWidget()->GetSelectedItems())
+	{
+		UFoliageType* FoliageType = TypeInfo->Settings;
+		if (FoliageType->IsAsset())
+		{
+			SelectedAssets.Add(FoliageType);
+		}
+		else
+		{
+			SelectedAssets.Add(FoliageType->GetStaticMesh());
+		}
+	}
+
+	if (SelectedAssets.Num())
+	{
+		GEditor->SyncBrowserToObjects(SelectedAssets);
+	}
+}
+
+void SFoliagePalette::OnSelectAllInstances()
+{
+	for (FFoliageMeshUIInfoPtr& TypeInfo : GetActiveViewWidget()->GetSelectedItems())
+	{
+		UFoliageType* FoliageType = TypeInfo->Settings;
+		FoliageEditMode->SelectInstances(FoliageType, true);
+	}
+}
+
+void SFoliagePalette::OnDeselectAllInstances()
+{
+	for (FFoliageMeshUIInfoPtr& TypeInfo : GetActiveViewWidget()->GetSelectedItems())
+	{
+		UFoliageType* FoliageType = TypeInfo->Settings;
+		FoliageEditMode->SelectInstances(FoliageType, false);
+	}
+}
+
+void SFoliagePalette::OnSelectInvalidInstances()
+{
+	for (FFoliageMeshUIInfoPtr& TypeInfo : GetActiveViewWidget()->GetSelectedItems())
+	{
+		const UFoliageType* FoliageType = TypeInfo->Settings;
+		FoliageEditMode->SelectInvalidInstances(FoliageType);
+	}
+}
+
+bool SFoliagePalette::CanSelectInstances() const
+{
+	return FoliageEditMode->UISettings.GetSelectToolSelected() || FoliageEditMode->UISettings.GetLassoSelectToolSelected();
+}
+
+// THUMBNAIL VIEW
+
+TSharedRef<ITableRow> SFoliagePalette::GenerateTile(FFoliageMeshUIInfoPtr Item, const TSharedRef<STableViewBase>& OwnerTable)
+{
+	TSharedPtr<FFoliagePaletteItemModel> ItemModel = MakeShareable(new FFoliagePaletteItemModel(Item, SharedThis(this), FoliageEditMode));
+	return SNew(SFoliagePaletteItemTile, OwnerTable, ItemModel);
+}
+
+// TREE VIEW
+
+TSharedRef<ITableRow> SFoliagePalette::TreeViewGenerateRow(FFoliageMeshUIInfoPtr Item, const TSharedRef<STableViewBase>& OwnerTable)
+{
+	TSharedPtr<FFoliagePaletteItemModel> ItemModel = MakeShareable(new FFoliagePaletteItemModel(Item, SharedThis(this), FoliageEditMode));
+	return SNew(SFoliagePaletteItemRow, OwnerTable, ItemModel);
+}
+
+void SFoliagePalette::TreeViewGetChildren(FFoliageMeshUIInfoPtr Item, TArray<FFoliageMeshUIInfoPtr>& OutChildren)
+{
+	//OutChildren = Item->GetChildren();
+}
+
+ECheckBoxState SFoliagePalette::GetState_AllMeshes() const
+{
+	bool bHasChecked = false;
+	bool bHasUnchecked = false;
+
+	for (const FFoliageMeshUIInfoPtr& TypeInfo : FoliageEditMode->GetFoliageMeshList())
+	{
+		if (TypeInfo->Settings->IsSelected)
+		{
+			bHasChecked = true;
+		}
+		else
+		{
+			bHasUnchecked = true;
+		}
+
+		if (bHasChecked && bHasUnchecked)
+		{
+			return ECheckBoxState::Undetermined;
+		}
+	}
+
+	return bHasChecked ? ECheckBoxState::Checked : ECheckBoxState::Unchecked;
+}
+
+void SFoliagePalette::OnCheckStateChanged_AllMeshes(ECheckBoxState InState)
+{
+	for (const FFoliageMeshUIInfoPtr& TypeInfo : FoliageEditMode->GetFoliageMeshList())
+	{
+		TypeInfo->Settings->IsSelected = (InState == ECheckBoxState::Checked);
+	}
+}
+
+FText SFoliagePalette::GetMeshesHeaderText() const
+{
+	int32 NumMeshes = FoliageEditMode->GetFoliageMeshList().Num();
+	return FText::Format(LOCTEXT("FoliageMeshCount", "Meshes ({0})"), FText::AsNumber(NumMeshes));
+}
+
+EColumnSortMode::Type SFoliagePalette::GetMeshColumnSortMode() const
+{
+	return FoliageEditMode->GetFoliageMeshListSortMode();
+}
+
+void SFoliagePalette::OnMeshesColumnSortModeChanged(EColumnSortPriority::Type InPriority, const FName& InColumnName, EColumnSortMode::Type InSortMode)
+{
+	FoliageEditMode->OnFoliageMeshListSortModeChanged(InSortMode);
+}
+
+FText SFoliagePalette::GetTotalInstanceCountTooltipText() const
+{
+	// Probably should cache these values, 
+	// but we call this only occasionally when tooltip is active
+	int32 InstanceCountTotal = 0;
+	int32 InstanceCountCurrentLevel = 0;
+	FoliageEditMode->CalcTotalInstanceCount(InstanceCountTotal, InstanceCountCurrentLevel);
+
+	return FText::Format(LOCTEXT("FoliageTotalInstanceCount", "Current Level: {0} Total: {1}"), FText::AsNumber(InstanceCountCurrentLevel), FText::AsNumber(InstanceCountTotal));
+}
+
+//	DETAILS VIEW
+
+void SFoliagePalette::RefreshMeshDetailsWidget()
+{
+	TArray<UObject*> SelectedFoliageTypes;
+
+	for (const auto& Info : GetActiveViewWidget()->GetSelectedItems())
+	{
+		SelectedFoliageTypes.Add(Info->Settings);
+	}
+
+	const bool bForceRefresh = true;
+	DetailsWidget->SetObjects(SelectedFoliageTypes, bForceRefresh);
+}
+
+FText SFoliagePalette::GetDetailsNameAreaText() const
+{
+	FText OutText;
+
+	auto SelectedTypes = GetActiveViewWidget()->GetSelectedItems();
+	if (SelectedTypes.Num() == 1)
+	{
+		FName DisplayName;
+		UFoliageType* SelectedType = SelectedTypes[0]->Settings;
+		if (SelectedType->IsAsset())
+		{
+			DisplayName = SelectedType->GetFName();
+		}
+		else
+		{
+			DisplayName = SelectedType->GetStaticMesh()->GetFName();
+		}
+		OutText = FText::FromName(DisplayName);
+	}
+	else if (SelectedTypes.Num() > 1)
+	{
+		OutText = FText::Format(LOCTEXT("DetailsNameAreaText_Multiple", "{0} Types Selected"), FText::AsNumber(SelectedTypes.Num()));
+	}
+
+	return OutText;
+}
+
+FText SFoliagePalette::GetShowHideDetailsTooltipText() const
+{
+	const bool bDetailsCurrentlyVisible = DetailsWidget->GetVisibility() != EVisibility::Collapsed;
+	return bDetailsCurrentlyVisible ? LOCTEXT("HideDetails_Tooltip", "Hide details for the selected foliage types.") : LOCTEXT("ShowDetails_Tooltip", "Show details for the selected foliage types.");
+}
+
+const FSlateBrush* SFoliagePalette::GetShowHideDetailsImage() const
+{
+	const bool bDetailsCurrentlyVisible = DetailsWidget->GetVisibility() != EVisibility::Collapsed;
+	return FEditorStyle::Get().GetBrush(bDetailsCurrentlyVisible ? "Symbols.DoubleDownArrow" : "Symbols.DoubleUpArrow");
+}
+
+FReply SFoliagePalette::OnShowHideDetailsClicked() const
+{
+	const bool bDetailsCurrentlyVisible = DetailsWidget->GetVisibility() != EVisibility::Collapsed;
+	DetailsWidget->SetVisibility(bDetailsCurrentlyVisible ? EVisibility::Collapsed : EVisibility::SelfHitTestInvisible);
+
+	return FReply::Handled();
+}
+
+#undef LOCTEXT_NAMESPACE
\ No newline at end of file
diff --git a/Engine/Source/Editor/FoliageEdit/Private/SFoliagePalette.h b/Engine/Source/Editor/FoliageEdit/Private/SFoliagePalette.h
new file mode 100644
--- /dev/null
+++ b/Engine/Source/Editor/FoliageEdit/Private/SFoliagePalette.h
@@ -0,0 +1,196 @@
+// Copyright 1998-2015 Epic Games, Inc. All Rights Reserved.
+
+#pragma once
+
+class SFoliagePaletteTile;
+class FEdModeFoliage;
+struct FFoliageMeshUIInfo;
+typedef TSharedPtr<FFoliageMeshUIInfo> FFoliageMeshUIInfoPtr; //should match typedef in FoliageEdMode.h
+
+typedef STreeView<FFoliageMeshUIInfoPtr> SFoliageTypeTreeView;
+typedef STileView<FFoliageMeshUIInfoPtr> SFoliageTypeTileView;
+
+/** View modes supported by the foliage palette */
+namespace EPaletteViewMode
+{
+	enum Type
+	{
+		Thumbnail,
+		Tree
+	};
+}
+
+/** The palette of foliage types available for use by the foliage edit mode */
+class SFoliagePalette : public SCompoundWidget
+{
+public:
+	SLATE_BEGIN_ARGS(SFoliagePalette) {}
+		SLATE_ARGUMENT(FEdModeFoliage*, FoliageEdMode)
+	SLATE_END_ARGS()
+
+	void Construct(const FArguments& InArgs);
+	void Refresh();
+	
+	bool AnySelectedTileHovered() const;
+	void ActivateAllSelectedTypes(bool bActivate) const;
+
+	/** @return True if the given view mode is the active view mode */
+	bool IsActiveViewMode(EPaletteViewMode::Type ViewMode) const;
+
+	/** @return True if tooltips should be shown when hovering over foliage type items in the palette */
+	bool ShouldShowTooltips() const;
+
+private:	// GENERAL
+
+	/** Binds commands used by the palette */
+	void BindCommands();
+
+	/** Adds the foliage type asset to the foliage actor's list of meshes. */
+	void AddFoliageType(const FAssetData& AssetData);
+	
+	/** Gets the asset picker for adding a foliage type. */
+	TSharedRef<SWidget> GetAddFoliageTypePicker();
+
+	/** Gets the visibility of the Add Foliage Type text in the header row button */
+	EVisibility GetAddFoliageTypeButtonTextVisibility() const;
+
+	/** Toggle whether all foliage types are active */
+	ECheckBoxState GetState_AllMeshes() const;
+	void OnCheckStateChanged_AllMeshes(ECheckBoxState InState);
+
+	/** Handler to trigger a refresh of the details view when the active tool changes */
+	void HandleOnToolChanged();
+
+	/** Sets the view mode of the palette */
+	void SetViewMode(EPaletteViewMode::Type NewViewMode);
+
+	/** Sets whether to show tooltips when hovering over foliage type items in the palette */
+	void ToggleShowTooltips();
+
+	/** Switches the palette display between the tile and tree view */
+	FReply OnToggleViewModeClicked();
+
+	/** @return The index of the view widget to display */
+	int32 GetActiveViewIndex() const;
+
+	/** Handler for selection changes in either view */
+	void OnSelectionChanged(FFoliageMeshUIInfoPtr Item, ESelectInfo::Type SelectInfo);
+
+	/** Toggle the activation state of a type on a double-click */
+	void OnItemDoubleClicked(FFoliageMeshUIInfoPtr Item) const;
+
+	/** Creates the view options menu */
+	TSharedRef<SWidget> GetViewOptionsMenuContent();
+
+	TSharedPtr<SListView<FFoliageMeshUIInfoPtr>> GetActiveViewWidget() const;
+
+	/** Gets the visibility of the "Drop Foliage Here" prompt for when the palette is empty */
+	EVisibility GetDropFoliageHintVisibility() const;
+
+private:	// CONTEXT MENU
+
+	/** @return the SWidget containing the context menu */
+	TSharedPtr<SWidget> ConstructFoliageTypeContextMenu();
+
+	/** Saves a temporary non-asset foliage type (created from a static mesh) as a foliage type asset */
+	void OnSaveSelectedAsFoliageType();
+
+	/** @return True if all selected types are non-assets */
+	bool OnCanSaveSelectedAsFoliageType() const;
+
+	/** @return True if any of the selected types are non-assets */
+	bool AreAnyNonAssetTypesSelected() const;
+
+	/** Fills 'Replace' menu command  */
+	void FillReplaceFoliageTypeSubmenu(FMenuBuilder& MenuBuilder);
+
+	/** Handler for 'Replace' command  */
+	void OnReplaceFoliageTypeSelected(const class FAssetData& AssetData);
+
+	/** Handler for 'Remove' command  */
+	void OnRemoveFoliageType();
+
+	/** Handler for 'Show in CB' command  */
+	void OnShowFoliageTypeInCB();
+
+	/** Handler for 'Select All' command  */
+	void OnSelectAllInstances();
+
+	/** Handler for 'Deselect All' command  */
+	void OnDeselectAllInstances();
+
+	/** Handler for 'Select Invalid Instances' command  */
+	void OnSelectInvalidInstances();
+
+	/** @return Whether selecting instances is currently possible */
+	bool CanSelectInstances() const;
+
+private:	// THUMBNAIL VIEW
+
+	/** Creates a thumbnail tile for the given foliage type */
+	TSharedRef<ITableRow> GenerateTile(FFoliageMeshUIInfoPtr Item, const TSharedRef<STableViewBase>& OwnerTable);
+
+private:	// TREE VIEW
+
+	/** Generates a row widget for foliage mesh item */
+	TSharedRef<ITableRow> TreeViewGenerateRow(FFoliageMeshUIInfoPtr Item, const TSharedRef<STableViewBase>& OwnerTable);
+
+	/** Generates a list of children items for foliage item */
+	void TreeViewGetChildren(FFoliageMeshUIInfoPtr Item, TArray<FFoliageMeshUIInfoPtr>& OutChildren);
+
+	/** Text for foliage meshes list header */
+	FText GetMeshesHeaderText() const;
+
+	/** Mesh list sorting support */
+	EColumnSortMode::Type GetMeshColumnSortMode() const;
+	void OnMeshesColumnSortModeChanged(EColumnSortPriority::Type InPriority, const FName& InColumnName, EColumnSortMode::Type InSortMode);
+
+	/** Tooltip text for 'Instance Count" column */
+	FText GetTotalInstanceCountTooltipText() const;
+
+private:	// DETAILS
+
+	/** Refreshes the mesh details widget to match the current selection */
+	void RefreshMeshDetailsWidget();
+
+	/** Gets the text for the details area header */
+	FText GetDetailsNameAreaText() const;
+
+	/** Gets the text for the show/hide details button tooltip */
+	FText GetShowHideDetailsTooltipText() const;
+
+	/** Gets the image for the show/hide details button */
+	const FSlateBrush* GetShowHideDetailsImage() const;
+
+	/** Handles the show/hide details button click */
+	FReply OnShowHideDetailsClicked() const;
+
+private:
+
+	/** Switches between the thumbnail and tree views */
+	TSharedPtr<class SWidgetSwitcher> WidgetSwitcher;
+
+	/** The Add Foliage Type combo button */
+	TSharedPtr<class SComboButton> AddFoliageTypeCombo;
+
+	/** The header row of the foliage mesh tree */
+	TSharedPtr<class SHeaderRow> TreeViewHeaderRow;
+
+	/** Foliage type thumbnails widget  */
+	TSharedPtr<SFoliageTypeTileView> TileViewWidget;
+
+	/** Foliage type tree widget  */
+	TSharedPtr<SFoliageTypeTreeView> TreeViewWidget;
+
+	/** Foliage mesh details widget  */
+	TSharedPtr<class IDetailsView> DetailsWidget;
+
+	/** Command list for binding functions for the context menu. */
+	TSharedPtr<FUICommandList> UICommandList;
+
+	FEdModeFoliage* FoliageEditMode;
+
+	//@todo: Should be saved to config
+	EPaletteViewMode::Type ActiveViewMode;
+	bool bShowTooltips;
+};
\ No newline at end of file
diff --git a/Engine/Source/Editor/UnrealEd/Private/AssetThumbnail.cpp b/Engine/Source/Editor/UnrealEd/Private/AssetThumbnail.cpp
--- a/Engine/Source/Editor/UnrealEd/Private/AssetThumbnail.cpp
+++ b/Engine/Source/Editor/UnrealEd/Private/AssetThumbnail.cpp
@@ -123,10 +123,12 @@ public:
 
 			TSharedPtr<SViewport> Viewport = 
 				SNew( SViewport )
-				.EnableGammaCorrection(false);
+				.EnableGammaCorrection(false)
+				.IgnoreTextureAlpha(false)
+				.EnableBlending(true);
 
 			Viewport->SetViewportInterface( AssetThumbnail.ToSharedRef() );
-			AssetThumbnail->GetViewportRenderTargetTexture(); // Access the render texture to push it on the stack if it isnt already rendered
+			AssetThumbnail->GetViewportRenderTargetTexture(); // Access the render texture to push it on the stack if it isn't already rendered
 
 			InArgs._ThumbnailPool->OnThumbnailRendered().AddSP(this, &SAssetThumbnail::OnThumbnailRendered);
 			InArgs._ThumbnailPool->OnThumbnailRenderFailed().AddSP(this, &SAssetThumbnail::OnThumbnailRenderFailed);
diff --git a/Engine/Source/Runtime/Slate/Private/Widgets/Layout/SSplitter.cpp b/Engine/Source/Runtime/Slate/Private/Widgets/Layout/SSplitter.cpp
--- a/Engine/Source/Runtime/Slate/Private/Widgets/Layout/SSplitter.cpp
+++ b/Engine/Source/Runtime/Slate/Private/Widgets/Layout/SSplitter.cpp
@@ -197,7 +197,7 @@ int32 SSplitter::OnPaint( const FPaintArgs& Args, const FGeometry& AllottedGeome
 				NormalHandleBrush,
 				MyClippingRect,
 				ShouldBeEnabled( bParentEnabled ),
-				InWidgetStyle.GetColorAndOpacityTint()
+				InWidgetStyle.GetColorAndOpacityTint() * NormalHandleBrush->TintColor.GetSpecifiedColor()
 			);
 		}
 		else
@@ -209,7 +209,7 @@ int32 SSplitter::OnPaint( const FPaintArgs& Args, const FGeometry& AllottedGeome
 				&Style->HandleHighlightBrush,
 				MyClippingRect,
 				ShouldBeEnabled( bParentEnabled ),
-				InWidgetStyle.GetColorAndOpacityTint()
+				InWidgetStyle.GetColorAndOpacityTint() * Style->HandleHighlightBrush.TintColor.GetSpecifiedColor()
 			);	
 		}
 	}
