Fix non-unity

diff --git a/Engine/Source/Editor/FoliageEdit/Private/FoliagePaletteItem.h b/Engine/Source/Editor/FoliageEdit/Private/FoliagePaletteItem.h
--- a/Engine/Source/Editor/FoliageEdit/Private/FoliagePaletteItem.h
+++ b/Engine/Source/Editor/FoliageEdit/Private/FoliagePaletteItem.h
@@ -27,7 +27,7 @@ public:
 	/** @return The foliage palette that contains the item */
 	const TSharedPtr<SFoliagePalette>& GetFoliagePalette() const;
 
-	const UFoliageType* GetFoliageType() const;
+	const class UFoliageType* GetFoliageType() const;
 
 	/** @return The thumbnail widget for this item */
 	TSharedRef<SWidget> GetThumbnailWidget() const;
diff --git a/Engine/Source/Editor/FoliageEdit/Private/FoliageTypePaintingCustomization.h b/Engine/Source/Editor/FoliageEdit/Private/FoliageTypePaintingCustomization.h
--- a/Engine/Source/Editor/FoliageEdit/Private/FoliageTypePaintingCustomization.h
+++ b/Engine/Source/Editor/FoliageEdit/Private/FoliageTypePaintingCustomization.h
@@ -1,9 +1,12 @@
 // Copyright 1998-2015 Epic Games, Inc. All Rights Reserved.
 #pragma once
 
+#include "IDetailCustomization.h"
+
 class FEdModeFoliage;
 class IDetailCategoryBuilder;
 class IDetailPropertyRow;
+class IPropertyHandle;
 
 /////////////////////////////////////////////////////
 // FFoliageTypePaintingCustomization
