Fixed: Foliage Painting tool density is not really working
UE-11495

diff --git a/Engine/Source/Editor/FoliageEdit/Private/FoliageEdMode.cpp b/Engine/Source/Editor/FoliageEdit/Private/FoliageEdMode.cpp
--- a/Engine/Source/Editor/FoliageEdit/Private/FoliageEdMode.cpp
+++ b/Engine/Source/Editor/FoliageEdit/Private/FoliageEdMode.cpp
@@ -438,10 +438,10 @@ void FEdModeFoliage::GetRandomVectorInBrush(FVector& OutStart, FVector& OutEnd)
 	FVector Point = Ru * U + Rv * V;
 
 	// find distance to surface of sphere brush from this point
-	float Rw = FMath::Sqrt(1.f - (FMath::Square(Ru) + FMath::Square(Rv)));
+	FVector Rw = FMath::Sqrt(1.f - (FMath::Square(Ru) + FMath::Square(Rv))) * BrushTraceDirection;
 
-	OutStart = BrushLocation + UISettings.GetRadius() * (Ru * U + Rv * V - Rw * BrushTraceDirection);
-	OutEnd = BrushLocation + UISettings.GetRadius() * (Ru * U + Rv * V + Rw * BrushTraceDirection);
+	OutStart	= BrushLocation + UISettings.GetRadius() * (Point - Rw);
+	OutEnd		= BrushLocation + UISettings.GetRadius() * (Point + Rw);
 }
 
 /** This does not check for overlaps or density */
@@ -493,39 +493,32 @@ static bool CheckForOverlappingSphere(const UWorld* InWorld, const UFoliageType*
 	return false;
 }
 
-static bool CheckLocationForPotentialInstance(const UWorld* InWorld, const UFoliageType* Settings, float DensityCheckRadius, const FVector& Location, const FVector& Normal, TArray<FVector>& PotentialInstanceLocations, FFoliageInstanceHash& PotentialInstanceHash)
+static bool CheckLocationForPotentialInstance(const UWorld* InWorld, const UFoliageType* Settings, const FVector& Location, const FVector& Normal, TArray<FVector>& PotentialInstanceLocations, FFoliageInstanceHash& PotentialInstanceHash)
 {
 	if (CheckLocationForPotentialInstance_ThreadSafe(Settings, Location, Normal) == false)
 	{
 		return false;
 	}
 
-	// Check existing instances. Use the Density radius rather than the minimum radius
-	if (CheckForOverlappingSphere(InWorld, Settings, FSphere(Location, DensityCheckRadius)))
-	{
-		return false;
-	}
-
 	// Check if we're too close to any other instances
 	if (Settings->Radius > 0.f)
 	{
-		// Check with other potential instances we're about to add.
-		bool bFoundOverlapping = false;
-		float RadiusSquared = FMath::Square(DensityCheckRadius/*Settings->Radius*/);
+		// Check existing instances. Use the Density radius rather than the minimum radius
+		if (CheckForOverlappingSphere(InWorld, Settings, FSphere(Location, Settings->Radius)))
+		{
+			return false;
+		}
 
+		// Check with other potential instances we're about to add.
+		const float RadiusSquared = FMath::Square(Settings->Radius);
 		auto TempInstances = PotentialInstanceHash.GetInstancesOverlappingBox(FBox::BuildAABB(Location, FVector(Settings->Radius)));
 		for (int32 Idx : TempInstances)
 		{
 			if ((PotentialInstanceLocations[Idx] - Location).SizeSquared() < RadiusSquared)
 			{
-				bFoundOverlapping = true;
-				break;
+				return false;
 			}
 		}
-		if (bFoundOverlapping)
-		{
-			return false;
-		}
 	}
 
 	int32 PotentialIdx = PotentialInstanceLocations.Add(Location);
@@ -725,9 +718,6 @@ void FEdModeFoliage::CalculatePotentialInstances(const UWorld* InWorld, const UF
 		Bucket.Reserve(DesiredInstances.Num());
 	}
 
-	// Radius where we expect to have a single instance, given the density rules
-	const float DensityCheckRadius = FMath::Max<float>(FMath::Sqrt((1000.f*1000.f) / (PI * Settings->Density)), Settings->Radius);
-
 	for (const FDesiredFoliageInstance& DesiredInst : DesiredInstances)
 	{
 		FHitResult Hit;
@@ -750,8 +740,8 @@ void FEdModeFoliage::CalculatePotentialInstances(const UWorld* InWorld, const UF
 			{
 				continue;
 			}
-			
-			const bool bValidInstance =	CheckLocationForPotentialInstance(InWorld, Settings, DensityCheckRadius, Hit.ImpactPoint, Hit.ImpactNormal, PotentialInstanceLocations, PotentialInstanceHash)
+
+			const bool bValidInstance =	CheckLocationForPotentialInstance(InWorld, Settings, Hit.ImpactPoint, Hit.ImpactNormal, PotentialInstanceLocations, PotentialInstanceHash)
 										&& VertexMaskCheck(Hit, Settings)
 										&& LandscapeLayerCheck(Hit, Settings, LocalCache, HitWeight);
 			if (bValidInstance)
@@ -1904,9 +1894,6 @@ void FEdModeFoliage::ApplyPaintBucket_Add(AActor* Actor)
 		FFoliageInstanceHash PotentialInstanceHash(7);	// use 128x128 cell size, as the brush radius is typically small.
 		TArray<FPotentialInstance> InstancesToPlace;
 
-		// Radius where we expect to have a single instance, given the density rules
-		const float DensityCheckRadius = FMath::Max<float>(FMath::Sqrt((1000.f*1000.f) / (PI * Settings->Density * UISettings.GetPaintDensity())), Settings->Radius);
-
 		for (TMap<UPrimitiveComponent*, TArray<FFoliagePaintBucketTriangle> >::TIterator ComponentIt(ComponentPotentialTriangles); ComponentIt; ++ComponentIt)
 		{
 			UPrimitiveComponent* Component = ComponentIt.Key();
@@ -1938,7 +1925,7 @@ void FEdModeFoliage::ApplyPaintBucket_Add(AActor* Actor)
 
 					// Check color mask and filters at this location
 					if (!CheckVertexColor(Settings, VertexColor) ||
-						!CheckLocationForPotentialInstance(World, Settings, DensityCheckRadius, InstLocation, Triangle.WorldNormal, PotentialInstanceLocations, PotentialInstanceHash))
+						!CheckLocationForPotentialInstance(World, Settings, InstLocation, Triangle.WorldNormal, PotentialInstanceLocations, PotentialInstanceHash))
 					{
 						continue;
 					}
