[UE-14309] - Saving view settings for the foliage palette across editor sessions
[UE-14461] - Added search bar that filters palette items
[UE-14155] - Mesh isn't added by proc foliage in addition to type asset when type asset is used
[UE-14528] - Instance count on the thumbnail disappears when thumbnail is too small and rounded off when value gets big

- Added thumbnail size slider to palette view options
- Add type/mesh drag-drop zone moved to be within the palette
- Ctrl + Wheel, Ctrl + MMB + drag, and "[" "]" adjust the size of the foliage brush in the viewport
- Key bindings for foliage commands work in the viewport
- Separate UI commands for palette item actions
- Rearranged view options and collapse details buttons

diff --git a/Engine/Source/Editor/FoliageEdit/Private/FoliageEdMode.cpp b/Engine/Source/Editor/FoliageEdit/Private/FoliageEdMode.cpp
--- a/Engine/Source/Editor/FoliageEdit/Private/FoliageEdMode.cpp
+++ b/Engine/Source/Editor/FoliageEdit/Private/FoliageEdMode.cpp
@@ -10,6 +10,7 @@
 #include "ModuleManager.h"
 #include "Editor/LevelEditor/Public/LevelEditor.h"
 #include "Toolkits/ToolkitManager.h"
+#include "FoliageEditActions.h"
 
 #include "Runtime/AssetRegistry/Public/AssetRegistryModule.h"
 
@@ -154,6 +155,7 @@ FEdModeFoliage::FEdModeFoliage()
 	: FEdMode()
 	, bToolActive(false)
 	, bCanAltDrag(false)
+	, bAdjustBrushRadius(false)
 	, FoliageMeshListSortMode(EColumnSortMode::None)
 {
 	// Load resources and construct brush component
@@ -175,8 +177,76 @@ FEdModeFoliage::FEdModeFoliage()
 
 	bBrushTraceValid = false;
 	BrushLocation = FVector::ZeroVector;
+
+	FFoliageEditCommands::Register();
+	UICommandList = MakeShareable(new FUICommandList);
+	BindCommands();
 }
 
+void FEdModeFoliage::BindCommands()
+{
+	const FFoliageEditCommands& Commands = FFoliageEditCommands::Get();
+
+	UICommandList->MapAction(
+		Commands.IncreaseBrushSize,
+		FExecuteAction::CreateRaw(this, &FEdModeFoliage::AdjustBrushRadius, 50.f),
+		FCanExecuteAction::CreateRaw(this, &FEdModeFoliage::CurrentToolUsesBrush));
+
+	UICommandList->MapAction(
+		Commands.DecreaseBrushSize,
+		FExecuteAction::CreateRaw(this, &FEdModeFoliage::AdjustBrushRadius, -50.f),
+		FCanExecuteAction::CreateRaw(this, &FEdModeFoliage::CurrentToolUsesBrush));
+
+	UICommandList->MapAction(
+		Commands.SetPaint,
+		FExecuteAction::CreateRaw(this, &FEdModeFoliage::OnSetPaint),
+		FCanExecuteAction(),
+		FIsActionChecked::CreateLambda([=]
+		{
+			return UISettings.GetPaintToolSelected();
+		}));
+
+	UICommandList->MapAction(
+		Commands.SetReapplySettings,
+		FExecuteAction::CreateRaw(this, &FEdModeFoliage::OnSetReapplySettings),
+		FCanExecuteAction(),
+		FIsActionChecked::CreateLambda([=]
+		{
+			return UISettings.GetReapplyToolSelected();
+		}));
+
+	UICommandList->MapAction(
+		Commands.SetSelect,
+		FExecuteAction::CreateRaw(this, &FEdModeFoliage::OnSetSelectInstance),
+		FCanExecuteAction(),
+		FIsActionChecked::CreateLambda([=]
+		{
+			return UISettings.GetSelectToolSelected();
+		}));
+
+	UICommandList->MapAction(
+		Commands.SetLassoSelect,
+		FExecuteAction::CreateRaw(this, &FEdModeFoliage::OnSetLasso),
+		FCanExecuteAction(),
+		FIsActionChecked::CreateLambda([=]
+		{
+			return UISettings.GetLassoSelectToolSelected();
+		}));
+
+	UICommandList->MapAction(
+		Commands.SetPaintBucket,
+		FExecuteAction::CreateRaw(this, &FEdModeFoliage::OnSetPaintFill),
+		FCanExecuteAction(),
+		FIsActionChecked::CreateLambda([=]
+		{
+			return UISettings.GetPaintBucketToolSelected();
+		}));
+}
+
+bool FEdModeFoliage::CurrentToolUsesBrush() const
+{
+	return UISettings.GetPaintToolSelected() || UISettings.GetReapplyToolSelected() || UISettings.GetLassoSelectToolSelected();
+}
 
 /** Destructor */
 FEdModeFoliage::~FEdModeFoliage()
@@ -291,7 +361,7 @@ void FEdModeFoliage::NotifyLevelRemovedFromWorld(ULevel* InLevel, UWorld* InWorl
 }
 
 /** When the user changes the current tool in the UI */
-void FEdModeFoliage::NotifyToolChanged()
+void FEdModeFoliage::HandleToolChanged()
 {
 	if (UISettings.GetSelectToolSelected() || UISettings.GetLassoSelectToolSelected())
 	{
@@ -305,6 +375,50 @@ void FEdModeFoliage::NotifyToolChanged()
 	OnToolChanged.Broadcast();
 }
 
+void FEdModeFoliage::ClearAllToolSelection()
+{
+	UISettings.SetLassoSelectToolSelected(false);
+	UISettings.SetPaintToolSelected(false);
+	UISettings.SetReapplyToolSelected(false);
+	UISettings.SetSelectToolSelected(false);
+	UISettings.SetPaintBucketToolSelected(false);
+}
+
+void FEdModeFoliage::OnSetPaint()
+{
+	ClearAllToolSelection();
+	UISettings.SetPaintToolSelected(true);
+	HandleToolChanged();
+}
+
+void FEdModeFoliage::OnSetReapplySettings()
+{
+	ClearAllToolSelection();
+	UISettings.SetReapplyToolSelected(true);
+	HandleToolChanged();
+}
+
+void FEdModeFoliage::OnSetSelectInstance()
+{
+	ClearAllToolSelection();
+	UISettings.SetSelectToolSelected(true);
+	HandleToolChanged();
+}
+
+void FEdModeFoliage::OnSetLasso()
+{
+	ClearAllToolSelection();
+	UISettings.SetLassoSelectToolSelected(true);
+	HandleToolChanged();
+}
+
+void FEdModeFoliage::OnSetPaintFill()
+{
+	ClearAllToolSelection();
+	UISettings.SetPaintBucketToolSelected(true);
+	HandleToolChanged();
+}
+
 bool FEdModeFoliage::DisallowMouseDeltaTracking() const
 {
 	// We never want to use the mouse delta tracker while painting
@@ -385,8 +499,13 @@ void FEdModeFoliage::FoliageBrushTrace(FEditorViewportClient* ViewportClient, in
 					(UISettings.bFilterTranslucent || !Material || !IsTranslucentBlendMode(Material->GetBlendMode())) &&
 					(CanPaint(PrimComp->GetComponentLevel())))
 				{
-					// Adjust the sphere brush
-					BrushLocation = Hit.Location;
+					if (!bAdjustBrushRadius)
+					{
+						// Adjust the brush location
+						BrushLocation = Hit.Location;
+					}
+					
+					// Still want to draw the brush when resizing
 					bBrushTraceValid = true;
 				}
 			}
@@ -786,19 +905,19 @@ static void SpawnFoliageInstance(UWorld* InWorld, const UFoliageType* Settings,
 FFoliageMeshInfo* UpdateMeshSettings(AInstancedFoliageActor* IFA, const FDesiredFoliageInstance& DesiredInstance)
 {
 	FFoliageMeshInfo* MeshInfo = nullptr;
-	//We have to add the mesh into the foliage actor. This assumes the foliage type is just a concrete InstancedStaticMesh.
-	//With the new tree system coming this will not be true and you'll need to traverse the tree and insert the needed meshes
-	if (const UFoliageType_InstancedStaticMesh* Settings = Cast<UFoliageType_InstancedStaticMesh>(DesiredInstance.FoliageType))
+
+	//We have to add the foliage type into the foliage actor.
+	if (const UFoliageType_InstancedStaticMesh* FoliageType = Cast<UFoliageType_InstancedStaticMesh>(DesiredInstance.FoliageType))
 	{
-		if (UStaticMesh* StaticMesh = Settings->GetStaticMesh())
+		if (UStaticMesh* StaticMesh = FoliageType->GetStaticMesh())
 		{
-			if (UFoliageType* ExistingSettings = IFA->GetSettingsForMesh(StaticMesh))
+			if (UFoliageType* ExistingType = IFA->GetSettingsForMesh(StaticMesh))
 			{
-				MeshInfo = IFA->UpdateMeshSettings(StaticMesh, Settings);
+				MeshInfo = IFA->UpdateMeshSettings(StaticMesh, FoliageType);
 			}
 			else
 			{
-				MeshInfo = IFA->AddMesh(Settings->GetStaticMesh(), nullptr, Settings);
+				IFA->AddFoliageType(FoliageType, &MeshInfo);
 			}
 		}
 	}
@@ -1268,6 +1387,21 @@ void FEdModeFoliage::SelectInvalidInstances(const UFoliageType* Settings)
 	}
 }
 
+void FEdModeFoliage::AdjustBrushRadius(float Adjustment)
+{
+	const float CurrentBrushRadius = UISettings.GetRadius();
+
+	if (Adjustment > 0.f)
+	{
+		UISettings.SetRadius(FMath::Min(CurrentBrushRadius + Adjustment, 8192.f));
+	}
+	else if (Adjustment < 0.f)
+	{
+		UISettings.SetRadius(FMath::Max(CurrentBrushRadius + Adjustment, 0.f));
+	}
+	
+}
+
 void FEdModeFoliage::ReapplyInstancesForBrush(UWorld* InWorld, const UFoliageType* Settings, const FSphere& BrushSphere, float Pressure)
 {
 	// Adjust instance density first
@@ -2535,6 +2669,15 @@ void FEdModeFoliage::ReallocateClusters(UFoliageType* Settings)
 /** FEdMode: Called when a key is pressed */
 bool FEdModeFoliage::InputKey(FEditorViewportClient* ViewportClient, FViewport* Viewport, FKey Key, EInputEvent Event)
 {
+	if (Event != IE_Released)
+	{
+		if (UICommandList->ProcessCommandBindings(Key, FSlateApplication::Get().GetModifierKeys(), false/*Event == IE_Repeat*/))
+		{
+			return true;
+		}
+	}
+	
+	bool bHandled = false;
 	if (UISettings.GetPaintToolSelected() || UISettings.GetReapplyToolSelected() || UISettings.GetLassoSelectToolSelected())
 	{
 		// Require Ctrl or not as per user preference
@@ -2558,12 +2701,12 @@ bool FEdModeFoliage::InputKey(FEditorViewportClient* ViewportClient, FViewport*
 					PreApplyBrush();
 					ApplyBrush(ViewportClient);
 					bToolActive = true;
-					return true;
+
+					bHandled = true;
 				}
 			}
 		}
-
-		if (bToolActive && Event == IE_Released &&
+		else if (bToolActive && Event == IE_Released &&
 			(Key == EKeys::LeftMouseButton || (FoliageEditorControlType == ELandscapeFoliageEditorControlType::RequireCtrl && (Key == EKeys::LeftControl || Key == EKeys::RightControl))))
 		{
 			//Set the cursor position to that of the slate cursor so it wont snap back
@@ -2572,28 +2715,47 @@ bool FEdModeFoliage::InputKey(FEditorViewportClient* ViewportClient, FViewport*
 			InstanceSnapshot.Empty();
 			LandscapeLayerCaches.Empty();
 			bToolActive = false;
-			return true;
+			
+			bHandled = true;
+		}
+		else if (IsCtrlDown(Viewport))
+		{
+			// Control + scroll adjusts the brush radius
+			static const float RadiusAdjustmentAmount = 25.f;
+			if (Key == EKeys::MouseScrollUp)
+			{
+				AdjustBrushRadius(RadiusAdjustmentAmount);
+				
+				bHandled = true;
+			}
+			else if (Key == EKeys::MouseScrollDown)
+			{
+				AdjustBrushRadius(-RadiusAdjustmentAmount);
+				
+				bHandled = true;
+			}
 		}
 	}
-
-	if (UISettings.GetSelectToolSelected() || UISettings.GetLassoSelectToolSelected())
+	else if (UISettings.GetSelectToolSelected() || UISettings.GetLassoSelectToolSelected())
 	{
 		if (Event == IE_Pressed)
 		{
 			if (Key == EKeys::Platform_Delete)
 			{
 				RemoveSelectedInstances(GetWorld());
-				return true;
+				
+				bHandled = true;
 			}
 			else if (Key == EKeys::End)
 			{
 				SnapSelectedInstancesToGround(GetWorld());
-				return true;
+				
+				bHandled = true;
 			}
 		}
 	}
 
-	return false;
+	return bHandled;
 }
 
 /** FEdMode: Render the foliage edit mode */
@@ -2716,7 +2878,12 @@ FVector FEdModeFoliage::GetWidgetLocation() const
 /** FEdMode: Called when a mouse button is pressed */
 bool FEdModeFoliage::StartTracking(FEditorViewportClient* InViewportClient, FViewport* InViewport)
 {
-	if (UISettings.GetSelectToolSelected() || UISettings.GetLassoSelectToolSelected())
+	if (IsCtrlDown(InViewport) && InViewport->KeyState(EKeys::MiddleMouseButton) && (UISettings.GetPaintToolSelected() || UISettings.GetReapplyToolSelected() || UISettings.GetLassoSelectToolSelected()))
+	{
+		bAdjustBrushRadius = true;
+		return true;
+	}
+	else if (UISettings.GetSelectToolSelected() || UISettings.GetLassoSelectToolSelected())
 	{
 		// Update pivot
 		UpdateWidgetLocationToInstanceSelection();
@@ -2733,29 +2900,42 @@ bool FEdModeFoliage::StartTracking(FEditorViewportClient* InViewportClient, FVie
 /** FEdMode: Called when the a mouse button is released */
 bool FEdModeFoliage::EndTracking(FEditorViewportClient* InViewportClient, FViewport* InViewport)
 {
-	if (UISettings.GetSelectToolSelected() || UISettings.GetLassoSelectToolSelected())
+	if (!bAdjustBrushRadius && (UISettings.GetSelectToolSelected() || UISettings.GetLassoSelectToolSelected()))
 	{
 		GEditor->EndTransaction();
 		return true;
 	}
+	else
+	{
+		bAdjustBrushRadius = false;
+		return true;
+	}
+
 	return FEdMode::EndTracking(InViewportClient, InViewport);
 }
 
 /** FEdMode: Called when mouse drag input it applied */
 bool FEdModeFoliage::InputDelta(FEditorViewportClient* InViewportClient, FViewport* InViewport, FVector& InDrag, FRotator& InRot, FVector& InScale)
 {
-	bool bFoundSelection = false;
-
-	bool bAltDown = InViewport->KeyState(EKeys::LeftAlt) || InViewport->KeyState(EKeys::RightAlt);
-
-	if (InViewportClient->GetCurrentWidgetAxis() != EAxisList::None && (UISettings.GetSelectToolSelected() || UISettings.GetLassoSelectToolSelected()))
+	if (bAdjustBrushRadius)
 	{
-		const bool bDuplicateInstances = (bAltDown && bCanAltDrag && (InViewportClient->GetCurrentWidgetAxis() & EAxisList::XYZ));
-		
-		TransformSelectedInstances(GetWorld(), InDrag, InRot, InScale, bDuplicateInstances);
+		if (UISettings.GetPaintToolSelected() || UISettings.GetReapplyToolSelected() || UISettings.GetLassoSelectToolSelected())
+		{
+			static const float RadiusAdjustmentFactor = 10.f;
+			AdjustBrushRadius(RadiusAdjustmentFactor * InDrag.Y);
+		}
+	}
+	else
+	{
+		if (InViewportClient->GetCurrentWidgetAxis() != EAxisList::None && (UISettings.GetSelectToolSelected() || UISettings.GetLassoSelectToolSelected()))
+		{
+			const bool bDuplicateInstances = (bCanAltDrag && IsAltDown(InViewport) && (InViewportClient->GetCurrentWidgetAxis() & EAxisList::XYZ));
 
-		// Only allow alt-drag on first InputDelta
-		bCanAltDrag = false;
+			TransformSelectedInstances(GetWorld(), InDrag, InRot, InScale, bDuplicateInstances);
+
+			// Only allow alt-drag on first InputDelta
+			bCanAltDrag = false;
+		}
 	}
 
 	return FEdMode::InputDelta(InViewportClient, InViewport, InDrag, InRot, InScale);
@@ -2817,6 +2997,15 @@ void FFoliageUISettings::Load()
 	GConfig->GetBool(TEXT("FoliageEdit"), TEXT("bFilterStaticMesh"), bFilterStaticMesh, GEditorPerProjectIni);
 	GConfig->GetBool(TEXT("FoliageEdit"), TEXT("bFilterBSP"), bFilterBSP, GEditorPerProjectIni);
 	GConfig->GetBool(TEXT("FoliageEdit"), TEXT("bFilterTranslucent"), bFilterTranslucent, GEditorPerProjectIni);
+
+	GConfig->GetBool(TEXT("FoliageEdit"), TEXT("bShowPaletteItemDetails"), bShowPaletteItemDetails, GEditorPerProjectIni);
+	GConfig->GetBool(TEXT("FoliageEdit"), TEXT("bShowPaletteItemTooltips"), bShowPaletteItemTooltips, GEditorPerProjectIni);
+
+	int32 ActivePaletteViewModeAsInt = 0;
+	GConfig->GetInt(TEXT("FoliageEdit"), TEXT("ActivePaletteViewMode"), ActivePaletteViewModeAsInt, GEditorPerProjectIni);
+	ActivePaletteViewMode = EFoliagePaletteViewMode::Type(ActivePaletteViewModeAsInt);
+	
+	GConfig->GetFloat(TEXT("FoliageEdit"), TEXT("PaletteThumbnailScale"), PaletteThumbnailScale, GEditorPerProjectIni);
 }
 
 /** Save UI settings to ini file */
@@ -2832,6 +3021,11 @@ void FFoliageUISettings::Save()
 	GConfig->SetBool(TEXT("FoliageEdit"), TEXT("bFilterStaticMesh"), bFilterStaticMesh, GEditorPerProjectIni);
 	GConfig->SetBool(TEXT("FoliageEdit"), TEXT("bFilterBSP"), bFilterBSP, GEditorPerProjectIni);
 	GConfig->SetBool(TEXT("FoliageEdit"), TEXT("bFilterTranslucent"), bFilterTranslucent, GEditorPerProjectIni);
+
+	GConfig->SetBool(TEXT("FoliageEdit"), TEXT("bShowPaletteItemDetails"), bShowPaletteItemDetails, GEditorPerProjectIni);
+	GConfig->SetBool(TEXT("FoliageEdit"), TEXT("bShowPaletteItemTooltips"), bShowPaletteItemTooltips, GEditorPerProjectIni);
+	GConfig->SetInt(TEXT("FoliageEdit"), TEXT("ActivePaletteViewMode"), ActivePaletteViewMode, GEditorPerProjectIni);
+	GConfig->SetFloat(TEXT("FoliageEdit"), TEXT("PaletteThumbnailScale"), PaletteThumbnailScale, GEditorPerProjectIni);
 }
 
 #undef LOCTEXT_NAMESPACE
\ No newline at end of file
diff --git a/Engine/Source/Editor/FoliageEdit/Private/FoliageEdMode.h b/Engine/Source/Editor/FoliageEdit/Private/FoliageEdMode.h
--- a/Engine/Source/Editor/FoliageEdit/Private/FoliageEdMode.h
+++ b/Engine/Source/Editor/FoliageEdit/Private/FoliageEdMode.h
@@ -14,6 +14,16 @@ class UFoliageType;
 //
 class UStaticMesh;
 
+/** View modes supported by the foliage palette */
+namespace EFoliagePaletteViewMode
+{
+	enum Type
+	{
+		Thumbnail,
+		Tree
+	};
+}
+
 // Current user settings in Foliage UI
 struct FFoliageUISettings
 {
@@ -53,6 +63,15 @@ struct FFoliageUISettings
 	bool GetFilterTranslucent() const { return bFilterTranslucent; }
 	void SetFilterTranslucent(bool InbFilterTranslucent) { bFilterTranslucent = InbFilterTranslucent; }
 
+	bool GetShowPaletteItemDetails() const { return bShowPaletteItemDetails; }
+	void SetShowPaletteItemDetails(bool InbShowPaletteItemDetails) { bShowPaletteItemDetails = InbShowPaletteItemDetails; }
+	bool GetShowPaletteItemTooltips() const { return bShowPaletteItemTooltips; }
+	void SetShowPaletteItemTooltips(bool InbShowPaletteItemTooltips) { bShowPaletteItemTooltips = InbShowPaletteItemTooltips; }
+	EFoliagePaletteViewMode::Type GetActivePaletteViewMode() const { return ActivePaletteViewMode; }
+	void SetActivePaletteViewMode(EFoliagePaletteViewMode::Type InActivePaletteViewMode) { ActivePaletteViewMode = InActivePaletteViewMode; }
+	float GetPaletteThumbnailScale() const { return PaletteThumbnailScale; }
+	void SetPaletteThumbnailScale(float InThumbnailScale) { PaletteThumbnailScale = InThumbnailScale; }
+
 	FFoliageUISettings()
 		: WindowX(-1)
 		, WindowY(-1)
@@ -64,6 +83,10 @@ struct FFoliageUISettings
 		, bLassoSelectToolSelected(false)
 		, bPaintBucketToolSelected(false)
 		, bReapplyPaintBucketToolSelected(false)
+		, bShowPaletteItemDetails(true)
+		, bShowPaletteItemTooltips(true)
+		, ActivePaletteViewMode(EFoliagePaletteViewMode::Thumbnail)
+		, PaletteThumbnailScale(0.3f)
 		, Radius(512.f)
 		, PaintDensity(0.5f)
 		, UnpaintDensity(0.f)
@@ -90,6 +113,11 @@ private:
 	bool bPaintBucketToolSelected;
 	bool bReapplyPaintBucketToolSelected;
 
+	bool bShowPaletteItemDetails;
+	bool bShowPaletteItemTooltips;
+	EFoliagePaletteViewMode::Type ActivePaletteViewMode;
+	float PaletteThumbnailScale;
+
 	float Radius;
 	float PaintDensity;
 	float UnpaintDensity;
@@ -170,6 +198,9 @@ class FEdModeFoliage : public FEdMode
 public:
 	FFoliageUISettings UISettings;
 
+	/** Command list lives here so that the key bindings on the commands can be processed in the viewport. */
+	TSharedPtr<FUICommandList> UICommandList;
+
 	/** Constructor */
 	FEdModeFoliage();
 
@@ -193,9 +224,6 @@ public:
 	void NotifyLevelAddedToWorld(ULevel* InLevel, UWorld* InWorld);
 	void NotifyLevelRemovedFromWorld(ULevel* InLevel, UWorld* InWorld);
 	
-	/** Called when the user changes the current tool in the UI */
-	void NotifyToolChanged();
-
 	/**
 	 * Called when the mouse is moved over the viewport
 	 *
@@ -335,6 +363,9 @@ public:
 	/** Find and select instances that don't have valid base or 'off-ground' */
 	void SelectInvalidInstances(const UFoliageType* Settings);
 
+	/** Adjusts the radius of the foliage brush by the specified amount */
+	void AdjustBrushRadius(float Adjustment);
+
 	/** Add desired instances. Uses foliage settings to determine location/scale/rotation and whether instances should be ignored */
 	static void AddInstances(UWorld* InWorld, const TArray<FDesiredFoliageInstance>& DesiredInstances);
 
@@ -343,6 +374,30 @@ public:
 	FSimpleMulticastDelegate OnToolChanged;
 private:
 
+	void BindCommands();
+	bool CurrentToolUsesBrush() const;
+
+	/** Called when the user changes the current tool in the UI */
+	void HandleToolChanged();
+
+	/** Deselects all tools */
+	void ClearAllToolSelection();
+
+	/** Sets the tool mode to Paint. */
+	void OnSetPaint();
+
+	/** Sets the tool mode to Reapply Settings. */
+	void OnSetReapplySettings();
+
+	/** Sets the tool mode to Select. */
+	void OnSetSelectInstance();
+
+	/** Sets the tool mode to Lasso Select. */
+	void OnSetLasso();
+
+	/** Sets the tool mode to Paint Bucket. */
+	void OnSetPaintFill();
+
 	/** Add instances inside the brush to match DesiredInstanceCount */
 	void AddInstancesForBrush(UWorld* InWorld, const UFoliageType* Settings, const FSphere& BrushSphere, int32 DesiredInstanceCount, float Pressure);
 	
@@ -414,6 +469,7 @@ private:
 
 	bool bToolActive;
 	bool bCanAltDrag;
+	bool bAdjustBrushRadius;
 
 	TArray<FFoliageMeshUIInfoPtr>	FoliageMeshList;
 	EColumnSortMode::Type			FoliageMeshListSortMode; 
diff --git a/Engine/Source/Editor/FoliageEdit/Private/FoliageEditActions.cpp b/Engine/Source/Editor/FoliageEdit/Private/FoliageEditActions.cpp
--- a/Engine/Source/Editor/FoliageEdit/Private/FoliageEditActions.cpp
+++ b/Engine/Source/Editor/FoliageEdit/Private/FoliageEditActions.cpp
@@ -10,6 +10,9 @@
 
 void FFoliageEditCommands::RegisterCommands()
 {
+	UI_COMMAND(DecreaseBrushSize, "Decrease Brush Size", "Decreases the size of the foliage brush", EUserInterfaceActionType::Button, FInputChord(EKeys::LeftBracket));
+	UI_COMMAND(IncreaseBrushSize, "Increase Brush Size", "Increases the size of the foliage brush", EUserInterfaceActionType::Button, FInputChord(EKeys::RightBracket));
+
 	UI_COMMAND( SetPaint, "Paint", "Paint", EUserInterfaceActionType::ToggleButton, FInputChord() );
 	UI_COMMAND( SetReapplySettings, "Reapply", "Reapply settings to instances", EUserInterfaceActionType::ToggleButton, FInputChord() );
 	UI_COMMAND( SetSelect, "Select", "Select", EUserInterfaceActionType::ToggleButton, FInputChord() );
@@ -19,12 +22,6 @@ void FFoliageEditCommands::RegisterCommands()
 	UI_COMMAND( SetNoSettings, "Hide Details", "Hide details.", EUserInterfaceActionType::ToggleButton, FInputChord() );
 	UI_COMMAND( SetPaintSettings, "Show Painting settings", "Show painting settings.", EUserInterfaceActionType::ToggleButton, FInputChord() );
 	UI_COMMAND( SetClusterSettings, "Show Instance settings", "Show settings for placed instances.", EUserInterfaceActionType::ToggleButton, FInputChord() );
-
-	UI_COMMAND( RemoveFoliageType, "Remove", "Remove this foliage type from the palette. Removes all associated instances as well.", EUserInterfaceActionType::Button, FInputChord() );
-	UI_COMMAND( ShowFoliageTypeInCB, "Show in Content Browser", "Show asset in Content Browser.", EUserInterfaceActionType::Button, FInputChord() );
-	UI_COMMAND( SelectAllInstances, "Select All Instances", "Select all instances of this foliage type (must be in a selection mode).", EUserInterfaceActionType::Button, FInputChord());
-	UI_COMMAND( DeselectAllInstances, "Deselect All Instances", "Deselect all instances of this foliage type (must be in a selection mode).", EUserInterfaceActionType::Button, FInputChord());
-	UI_COMMAND( SelectInvalidInstances, "Select Invalid Instances", "Select all instances of this foliage type that are off ground (must be in a selection mode).", EUserInterfaceActionType::Button, FInputChord());
 }
 
 #undef LOCTEXT_NAMESPACE
\ No newline at end of file
diff --git a/Engine/Source/Editor/FoliageEdit/Private/FoliageEditActions.h b/Engine/Source/Editor/FoliageEdit/Private/FoliageEditActions.h
--- a/Engine/Source/Editor/FoliageEdit/Private/FoliageEditActions.h
+++ b/Engine/Source/Editor/FoliageEdit/Private/FoliageEditActions.h
@@ -1,7 +1,5 @@
 // Copyright 1998-2015 Epic Games, Inc. All Rights Reserved.
 
-
-
 #pragma once
 
 /**
@@ -25,6 +23,10 @@ public:
 	 * Foliage Edit Commands
 	 */
 	
+	/** Commands for the foliage brush settings. */
+	TSharedPtr< FUICommandInfo > IncreaseBrushSize;
+	TSharedPtr< FUICommandInfo > DecreaseBrushSize;
+
 	/** Commands for the tools toolbar. */
 	TSharedPtr< FUICommandInfo > SetPaint;
 	TSharedPtr< FUICommandInfo > SetReapplySettings;
diff --git a/Engine/Source/Editor/FoliageEdit/Private/FoliagePaletteCommands.cpp b/Engine/Source/Editor/FoliageEdit/Private/FoliagePaletteCommands.cpp
new file mode 100644
--- /dev/null
+++ b/Engine/Source/Editor/FoliageEdit/Private/FoliagePaletteCommands.cpp
@@ -0,0 +1,17 @@
+// Copyright 1998-2015 Epic Games, Inc. All Rights Reserved.
+
+#include "UnrealEd.h"
+#include "FoliagePaletteCommands.h"
+
+#define LOCTEXT_NAMESPACE ""
+
+void FFoliagePaletteCommands::RegisterCommands()
+{
+	UI_COMMAND(RemoveFoliageType, "Remove", "Remove this foliage type from the palette. Removes all associated instances as well.", EUserInterfaceActionType::Button, FInputChord(EKeys::Delete));
+	UI_COMMAND(ShowFoliageTypeInCB, "Show in Content Browser", "Show asset in Content Browser.", EUserInterfaceActionType::Button, FInputChord());
+	UI_COMMAND(SelectAllInstances, "Select All Instances", "Select all instances of this foliage type (must be in a selection mode).", EUserInterfaceActionType::Button, FInputChord());
+	UI_COMMAND(DeselectAllInstances, "Deselect All Instances", "Deselect all instances of this foliage type (must be in a selection mode).", EUserInterfaceActionType::Button, FInputChord());
+	UI_COMMAND(SelectInvalidInstances, "Select Invalid Instances", "Select all instances of this foliage type that are off ground (must be in a selection mode).", EUserInterfaceActionType::Button, FInputChord());
+}
+
+#undef LOCTEXT_NAMESPACE
\ No newline at end of file
diff --git a/Engine/Source/Editor/FoliageEdit/Private/FoliagePaletteCommands.h b/Engine/Source/Editor/FoliageEdit/Private/FoliagePaletteCommands.h
new file mode 100644
--- /dev/null
+++ b/Engine/Source/Editor/FoliageEdit/Private/FoliagePaletteCommands.h
@@ -0,0 +1,31 @@
+// Copyright 1998-2015 Epic Games, Inc. All Rights Reserved.
+
+#pragma once
+
+/**
+* Foliage Mode Palette Actions
+*/
+class FFoliagePaletteCommands : public TCommands < FFoliagePaletteCommands >
+{
+
+public:
+	FFoliagePaletteCommands() : TCommands<FFoliagePaletteCommands>
+		(
+		"FoliagePalette", // Context name for fast lookup
+		NSLOCTEXT("Contexts", "FoliagePalette", "Foliage Palette"), // Localized context name for displaying
+		"FoliageEditMode", // Parent
+		FEditorStyle::GetStyleSetName() // Icon Style Set
+		)
+	{
+	}
+
+	/** Initialize commands */
+	virtual void RegisterCommands() override;
+
+	/** Foliage Palette Commands */
+	TSharedPtr< FUICommandInfo > RemoveFoliageType;
+	TSharedPtr< FUICommandInfo > ShowFoliageTypeInCB;
+	TSharedPtr< FUICommandInfo > SelectAllInstances;
+	TSharedPtr< FUICommandInfo > DeselectAllInstances;
+	TSharedPtr< FUICommandInfo > SelectInvalidInstances;
+};
diff --git a/Engine/Source/Editor/FoliageEdit/Private/FoliagePaletteItem.cpp b/Engine/Source/Editor/FoliageEdit/Private/FoliagePaletteItem.cpp
--- a/Engine/Source/Editor/FoliageEdit/Private/FoliagePaletteItem.cpp
+++ b/Engine/Source/Editor/FoliageEdit/Private/FoliagePaletteItem.cpp
@@ -13,6 +13,17 @@
 
 #define LOCTEXT_NAMESPACE "FoliageEd_Mode"
 
+namespace EInstanceCountMagnitude
+{
+	enum Type
+	{
+		Thousand,
+		Million,
+		Billion,
+		Max
+	};
+}
+
 FFoliagePaletteItemModel::FFoliagePaletteItemModel(FFoliageMeshUIInfoPtr InTypeInfo, TSharedRef<SFoliagePalette> InFoliagePalette, FEdModeFoliage* InFoliageEditMode)
 	: TypeInfo(InTypeInfo)
 	, FoliagePalette(InFoliagePalette)
@@ -22,11 +33,12 @@ FFoliagePaletteItemModel::FFoliagePaletteItemModel(FFoliageMeshUIInfoPtr InTypeI
 	TSharedPtr<FAssetThumbnailPool> ThumbnailPool = LevelEditorModule.GetFirstLevelEditor()->GetThumbnailPool();
 	FAssetData AssetData = FAssetData(TypeInfo->Settings->GetStaticMesh());
 
-	TSharedPtr< FAssetThumbnail > Thumbnail = MakeShareable(new FAssetThumbnail(AssetData, 128, 128, ThumbnailPool));
+	int32 MaxThumbnailSize = FoliagePaletteConstants::ThumbnailSizeRange.Max;
+	TSharedPtr< FAssetThumbnail > Thumbnail = MakeShareable(new FAssetThumbnail(AssetData, MaxThumbnailSize, MaxThumbnailSize, ThumbnailPool));
 	FAssetThumbnailConfig ThumbnailConfig;
 	if (TypeInfo->Settings->IsAsset())
 	{
-		//@todo: access the foliage type asset type action and get the color from that
+		//@todo: access the foliage type asset type action and get the color from that (and maybe change it from pale blue - hard to see against the white thumbnail)
 		ThumbnailConfig.AssetTypeColorOverride = FColor(128, 192, 255);
 	}
 
@@ -38,11 +50,21 @@ const TSharedPtr<SFoliagePalette>& FFoliagePaletteItemModel::GetFoliagePalette()
 	return FoliagePalette;
 }
 
-const UFoliageType* FFoliagePaletteItemModel::GetFoliageType() const
+FFoliageMeshUIInfoPtr FFoliagePaletteItemModel::GetTypeUIInfo() const
+{
+	return TypeInfo;
+}
+
+UFoliageType* FFoliagePaletteItemModel::GetFoliageType() const
 {
 	return TypeInfo->Settings;
 }
 
+const FFoliageUISettings& FFoliagePaletteItemModel::GetFoliageUISettings() const
+{
+	return FoliageEditMode->UISettings;
+}
+
 TSharedRef<SWidget> FFoliagePaletteItemModel::GetThumbnailWidget() const
 {
 	return ThumbnailWidget.ToSharedRef();
@@ -71,8 +93,9 @@ TSharedRef<SToolTip> FFoliagePaletteItemModel::CreateTooltipWidget() const
 					.BorderImage(FEditorStyle::GetBrush("ContentBrowser.TileViewTooltip.ContentBorder"))
 					[
 						SNew(STextBlock)
-						.Text(this, &FFoliagePaletteItemModel::GetFoliageTypeNameText)
+						.Text(this, &FFoliagePaletteItemModel::GetFoliageTypeDisplayNameText)
 						.Font(FEditorStyle::GetFontStyle("ContentBrowser.TileViewTooltip.NameFont"))
+						.HighlightText(this, &FFoliagePaletteItemModel::GetPaletteSearchText)
 					]
 				]
 
@@ -80,19 +103,57 @@ TSharedRef<SToolTip> FFoliagePaletteItemModel::CreateTooltipWidget() const
 				.AutoHeight()
 				.Padding(FMargin(0.f, 3.f, 0.f, 0.f))
 				[
-					SNew(SBorder)
-					.Visibility(this, &FFoliagePaletteItemModel::GetTooltipThumbnailVisibility)
-					.Padding(6.f)
-					.HAlign(HAlign_Left)
-					.BorderImage(FEditorStyle::GetBrush("ContentBrowser.TileViewTooltip.ContentBorder"))
+					
+					SNew(SHorizontalBox)
+
+					+ SHorizontalBox::Slot()
+					.AutoWidth()
+					.Padding(FMargin(0.f, 0.f, 3.f, 0.f))
 					[
-						SNew(SBox)
-						.HeightOverride(64.f)
-						.WidthOverride(64.f)
+						SNew(SBorder)
+						.Padding(6.f)
+						.HAlign(HAlign_Center)
+						.Visibility(this, &FFoliagePaletteItemModel::GetTooltipThumbnailVisibility)
+						.BorderImage(FEditorStyle::GetBrush("ContentBrowser.TileViewTooltip.ContentBorder"))
 						[
-							GetThumbnailWidget()
+							SNew(SBox)
+							.HeightOverride(64.f)
+							.WidthOverride(64.f)
+							[
+								GetThumbnailWidget()
+							]
 						]
 					]
+
+					+ SHorizontalBox::Slot()
+					[
+						SNew(SBorder)
+						.Padding(6.f)
+						.BorderImage(FEditorStyle::GetBrush("ContentBrowser.TileViewTooltip.ContentBorder"))
+						[
+							SNew(SVerticalBox)
+
+							+ SVerticalBox::Slot()
+							[
+								SNew(SHorizontalBox)
+								
+								+ SHorizontalBox::Slot()
+								.AutoWidth()
+								[
+									SNew(STextBlock)
+									.Text(LOCTEXT("InstanceCountHeading", "Instance Count: "))
+									.ColorAndOpacity(FSlateColor::UseSubduedForeground())
+								]
+
+								+ SHorizontalBox::Slot()
+								.AutoWidth()
+								[
+									SNew(STextBlock)
+									.Text(this, &FFoliagePaletteItemModel::GetInstanceCountText, false)
+								]
+							]
+						]
+					]					
 				]
 			]
 		];
@@ -128,7 +189,7 @@ TSharedRef<SButton> FFoliagePaletteItemModel::CreateSaveAssetButton(TAttribute<E
 		];
 }
 
-FText FFoliagePaletteItemModel::GetFoliageTypeNameText() const
+FText FFoliagePaletteItemModel::GetFoliageTypeDisplayNameText() const
 {
 	if (TypeInfo->Settings->IsAsset())
 	{
@@ -140,18 +201,59 @@ FText FFoliagePaletteItemModel::GetFoliageTypeNameText() const
 	}
 }
 
-FText FFoliagePaletteItemModel::GetInstanceCountText() const
+FText FFoliagePaletteItemModel::GetPaletteSearchText() const
+{
+	return FoliagePalette->GetSearchText();
+}
+
+FText FFoliagePaletteItemModel::GetInstanceCountText(bool bRounded) const
 {
 	const int32	InstanceCountTotal = TypeInfo->InstanceCountTotal;
 	const int32	InstanceCountCurrentLevel = TypeInfo->InstanceCountCurrentLevel;
 
 	if (InstanceCountCurrentLevel != InstanceCountTotal)
 	{
-		return FText::Format(LOCTEXT("InstanceCount", "{0} ({1})"), FText::AsNumber(InstanceCountCurrentLevel), FText::AsNumber(InstanceCountTotal));
+		return FText::Format(LOCTEXT("InstanceCount_Total", "{0} ({1})"), FText::AsNumber(InstanceCountCurrentLevel), FText::AsNumber(InstanceCountTotal));
 	}
-	else
+	else if (bRounded)
+	{
+		// Note: Instance counts greater than 999 billion (unlikely) will not be formatted properly
+		static const FText Suffixes[EInstanceCountMagnitude::Max] =
+		{
+			LOCTEXT("Suffix_Thousand", "K"),
+			LOCTEXT("Suffix_Million", "M"),
+			LOCTEXT("Suffix_Billion", "B")
+		};
+
+		int32 NumThousands = 0;
+		float DisplayValue = InstanceCountCurrentLevel;
+		while (DisplayValue >= 1000.f && NumThousands < EInstanceCountMagnitude::Max)
+		{
+			DisplayValue /= 1000.f;
+			NumThousands++;
+		}
+
+		if (NumThousands > 0)
+		{
+			// Allow 3 significant figures
+			const int32 NumFractionalDigits = DisplayValue >= 100.f ? 0 : DisplayValue >= 10.f ? 1 : 2;
+
+			FNumberFormattingOptions Options;
+			Options.SetMaximumFractionalDigits(NumFractionalDigits);
+			return FText::Format(LOCTEXT("InstanceCount_CurrentLevel", "{0}{1}"), FText::AsNumber(DisplayValue, &Options), Suffixes[NumThousands - 1]);
+		}
+	}
+
+	return FText::AsNumber(InstanceCountCurrentLevel);
+}
+
+void FFoliagePaletteItemModel::SetTypeActiveInPalette(bool bSetActiveInPalette)
+{
+	UFoliageType* FoliageType = GetFoliageType();
+	if (FoliageType->IsSelected != bSetActiveInPalette)
 	{
-		return FText::AsNumber(InstanceCountCurrentLevel);
+		FoliageType->Modify();
+		FoliageType->IsSelected = bSetActiveInPalette;
 	}
 }
 
@@ -162,11 +264,7 @@ void FFoliagePaletteItemModel::HandleCheckStateChanged(const ECheckBoxState NewC
 	const bool bShouldActivate = NewCheckedState == ECheckBoxState::Checked;
 	if (!IsItemWidgetSelected.Get())
 	{
-		if (TypeInfo->Settings->IsSelected != bShouldActivate)
-		{
-			TypeInfo->Settings->Modify();
-			TypeInfo->Settings->IsSelected = bShouldActivate;
-		}
+		SetTypeActiveInPalette(bShouldActivate);
 	}
 	else
 	{
@@ -181,7 +279,7 @@ ECheckBoxState FFoliagePaletteItemModel::GetCheckBoxState() const
 
 bool FFoliagePaletteItemModel::IsSaveEnabled() const
 {
-	UFoliageType* FoliageType = TypeInfo->Settings;
+	UFoliageType* FoliageType = GetFoliageType();
 
 	// Saving is enabled for non-assets and dirty assets
 	return (!FoliageType->IsAsset() || FoliageType->GetOutermost()->IsDirty());
@@ -205,13 +303,15 @@ EVisibility FFoliagePaletteItemModel::GetTooltipVisibility() const
 
 EVisibility FFoliagePaletteItemModel::GetTooltipThumbnailVisibility() const
 {
-	return FoliagePalette->IsActiveViewMode(EPaletteViewMode::Tree) ? EVisibility::SelfHitTestInvisible : EVisibility::Collapsed;
+	return FoliagePalette->IsActiveViewMode(EFoliagePaletteViewMode::Tree) ? EVisibility::SelfHitTestInvisible : EVisibility::Collapsed;
 }
 
 ////////////////////////////////////////////////
 // SFoliagePaletteItemTile
 ////////////////////////////////////////////////
 
+const float SFoliagePaletteItemTile::MinScaleForOverlayItems = 0.2f;
+
 void SFoliagePaletteItemTile::Construct(const FArguments& InArgs, TSharedRef<STableViewBase> InOwnerTableView, TSharedPtr<FFoliagePaletteItemModel>& InModel)
 {
 	Model = InModel;
@@ -229,6 +329,7 @@ void SFoliagePaletteItemTile::Construct(const FArguments& InArgs, TSharedRef<STa
 			SNew(SOverlay)
 			.ToolTip(Model->CreateTooltipWidget())
 			
+			// Thumbnail
 			+ SOverlay::Slot()
 			[
 				SNew(SBorder)
@@ -241,6 +342,7 @@ void SFoliagePaletteItemTile::Construct(const FArguments& InArgs, TSharedRef<STa
 				]
 			]
 			
+			// Checkbox
 			+ SOverlay::Slot()
 			.HAlign(HAlign_Left)
 			.VAlign(VAlign_Top)
@@ -256,6 +358,7 @@ void SFoliagePaletteItemTile::Construct(const FArguments& InArgs, TSharedRef<STa
 				]
 			]
 
+			// Save Button
 			+ SOverlay::Slot()
 			.HAlign(HAlign_Right)
 			.VAlign(VAlign_Top)
@@ -264,15 +367,15 @@ void SFoliagePaletteItemTile::Construct(const FArguments& InArgs, TSharedRef<STa
 				Model->CreateSaveAssetButton(SaveButtonVisibility)
 			]
 
-			+ SOverlay::Slot()
-
+			// Instance Count
 			+ SOverlay::Slot()
 			.HAlign(HAlign_Right)
 			.VAlign(VAlign_Bottom)
 			.Padding(FMargin(6.f, 8.f))
 			[
 				SNew(STextBlock)
-				.Text(Model.Get(), &FFoliagePaletteItemModel::GetInstanceCountText)
+				.Visibility(this, &SFoliagePaletteItemTile::GetInstanceCountVisibility)
+				.Text(Model.ToSharedRef(), &FFoliagePaletteItemModel::GetInstanceCountText, true)
 				.ShadowOffset(FVector2D(1.f, 1.f))
 				.ColorAndOpacity(FLinearColor(.85f, .85f, .85f, 1.f))
 			]
@@ -292,7 +395,17 @@ EVisibility SFoliagePaletteItemTile::GetCheckBoxVisibility() const
 
 EVisibility SFoliagePaletteItemTile::GetSaveButtonVisibility() const
 {
-	return IsHovered() ? EVisibility::Visible : EVisibility::Collapsed;
+	return IsHovered() && CanShowOverlayItems() ? EVisibility::Visible : EVisibility::Collapsed;
+}
+
+EVisibility SFoliagePaletteItemTile::GetInstanceCountVisibility() const
+{
+	return CanShowOverlayItems() ? EVisibility::Visible : EVisibility::Collapsed;
+}
+
+bool SFoliagePaletteItemTile::CanShowOverlayItems() const
+{
+	return Model->GetFoliageUISettings().GetPaletteThumbnailScale() >= MinScaleForOverlayItems;
 }
 
 ////////////////////////////////////////////////
@@ -332,7 +445,8 @@ TSharedRef<SWidget> SFoliagePaletteItemRow::GenerateWidgetForColumn(const FName&
 			.AutoWidth()
 			[
 				SNew(STextBlock)
-				.Text(Model.Get(), &FFoliagePaletteItemModel::GetFoliageTypeNameText)
+				.Text(Model.ToSharedRef(), &FFoliagePaletteItemModel::GetFoliageTypeDisplayNameText)
+				.HighlightText(Model.ToSharedRef(), &FFoliagePaletteItemModel::GetPaletteSearchText)
 			];
 	}
 	else if (ColumnName == FoliagePaletteTreeColumns::ColumnID_InstanceCount)
@@ -346,7 +460,7 @@ TSharedRef<SWidget> SFoliagePaletteItemRow::GenerateWidgetForColumn(const FName&
 			.AutoWidth()
 			[
 				SNew(STextBlock)
-				.Text(Model.Get(), &FFoliagePaletteItemModel::GetInstanceCountText)
+				.Text(Model.ToSharedRef(), &FFoliagePaletteItemModel::GetInstanceCountText, true)
 			];
 	}
 	else if (ColumnName == FoliagePaletteTreeColumns::ColumnID_Save)
diff --git a/Engine/Source/Editor/FoliageEdit/Private/FoliagePaletteItem.h b/Engine/Source/Editor/FoliageEdit/Private/FoliagePaletteItem.h
--- a/Engine/Source/Editor/FoliageEdit/Private/FoliagePaletteItem.h
+++ b/Engine/Source/Editor/FoliageEdit/Private/FoliagePaletteItem.h
@@ -4,6 +4,7 @@
 
 class SToolTip;
 class SFoliagePalette;
+class UFoliageType;
 
 class FEdModeFoliage;
 struct FFoliageMeshUIInfo;
@@ -27,7 +28,9 @@ public:
 	/** @return The foliage palette that contains the item */
 	const TSharedPtr<SFoliagePalette>& GetFoliagePalette() const;
 
-	const class UFoliageType* GetFoliageType() const;
+	FFoliageMeshUIInfoPtr GetTypeUIInfo() const;
+	UFoliageType* GetFoliageType() const;
+	const struct FFoliageUISettings& GetFoliageUISettings() const;
 
 	/** @return The thumbnail widget for this item */
 	TSharedRef<SWidget> GetThumbnailWidget() const;
@@ -42,10 +45,16 @@ public:
 	TSharedRef<class SButton> CreateSaveAssetButton(TAttribute<EVisibility> InVisibility = TAttribute<EVisibility>());
 
 	/** Gets the name to display of the foliage type object as text */
-	FText GetFoliageTypeNameText() const;
+	FText GetFoliageTypeDisplayNameText() const;
 
-	/** Gets the number of instances of this foliage type exist in the level */
-	FText GetInstanceCountText() const;
+	/** Gets the current search filter text */
+	FText GetPaletteSearchText() const;
+
+	/** Gets the number of instances of this foliage type  */
+	FText GetInstanceCountText(bool bLevelOnly = true) const;
+
+	/** Sets whether this foliage type is active in the palette */
+	void SetTypeActiveInPalette(bool bSetActiveInPalette);
 
 private:
 	/** Handles the change in activation of the item in the palette */
@@ -87,8 +96,11 @@ private:
 	FLinearColor GetTileColorAndOpacity() const;
 	EVisibility GetCheckBoxVisibility() const;
 	EVisibility GetSaveButtonVisibility() const;
+	EVisibility GetInstanceCountVisibility() const;
+	bool CanShowOverlayItems() const;
 
 private:
+	static const float MinScaleForOverlayItems;
 	TSharedPtr<FFoliagePaletteItemModel> Model;
 };
 
diff --git a/Engine/Source/Editor/FoliageEdit/Private/SFoliageEdit.cpp b/Engine/Source/Editor/FoliageEdit/Private/SFoliageEdit.cpp
--- a/Engine/Source/Editor/FoliageEdit/Private/SFoliageEdit.cpp
+++ b/Engine/Source/Editor/FoliageEdit/Private/SFoliageEdit.cpp
@@ -11,88 +11,13 @@
 #include "SFoliagePalette.h"
 #include "SScaleBox.h"
 
-
-// Possibly not needed (move over to SFoliagePalette)
-#include "Editor/UnrealEd/Public/AssetSelection.h"
-#include "Editor/UnrealEd/Public/ScopedTransaction.h"
-
-// Seemingly totally unneeded
-//#include "DetailLayoutBuilder.h"
-
-
 #define LOCTEXT_NAMESPACE "FoliageEd_Mode"
 
-class SFoliageDragDropHandler : public SCompoundWidget
-{
-public:
-	SLATE_BEGIN_ARGS(SFoliageDragDropHandler) {}
-	SLATE_DEFAULT_SLOT(FArguments, Content)
-	SLATE_ARGUMENT(TWeakPtr<SFoliageEdit>, FoliageEditPtr)
-	SLATE_END_ARGS()
-
-	/** SCompoundWidget functions */
-	void Construct(const FArguments& InArgs)
-	{
-		FoliageEditPtr = InArgs._FoliageEditPtr;
-		bIsDragOn = false;
-
-		this->ChildSlot
-		[
-			SNew(SBorder)
-			.BorderImage(FEditorStyle::GetBrush("WhiteBrush"))
-			.BorderBackgroundColor(this, &SFoliageDragDropHandler::GetBackgroundColor)
-			.Padding(FMargin(30))
-			[
-				InArgs._Content.Widget
-			]
-		];
-	}
-
-	~SFoliageDragDropHandler() {}
-
-	FReply OnDrop(const FGeometry& MyGeometry, const FDragDropEvent& DragDropEvent) override
-	{
-		bIsDragOn = false;
-
-		TSharedPtr<SFoliageEdit> PinnedPtr = FoliageEditPtr.Pin();
-		if (PinnedPtr.IsValid())
-		{
-			return PinnedPtr->OnDrop_ListView(MyGeometry, DragDropEvent);
-		}
-		return FReply::Handled();
-	}
-
-	virtual void OnDragEnter(const FGeometry& MyGeometry, const FDragDropEvent& DragDropEvent) override
-	{
-		bIsDragOn = true;
-	}
-
-	virtual void OnDragLeave(const FDragDropEvent& DragDropEvent) override
-	{
-		bIsDragOn = false;
-	}
-
-private:
-	FSlateColor GetBackgroundColor() const
-	{
-		return bIsDragOn ? FLinearColor(1.0f, 0.6f, 0.1f, 0.9f) : FLinearColor(0.1f, 0.1f, 0.1f, 0.9f);
-	}
-
-private:
-	TWeakPtr<SFoliageEdit> FoliageEditPtr;
-	bool bIsDragOn;
-};
-
-
 BEGIN_SLATE_FUNCTION_BUILD_OPTIMIZATION
 void SFoliageEdit::Construct(const FArguments& InArgs)
 {
 	FoliageEditMode = (FEdModeFoliage*)GLevelEditorModeTools().GetActiveMode(FBuiltinEditorModes::EM_Foliage);
 
-	FFoliageEditCommands::Register();
-	UICommandList = MakeShareable(new FUICommandList);
-	BindCommands();
-
 	IIntroTutorials& IntroTutorials = FModuleManager::LoadModuleChecked<IIntroTutorials>(TEXT("IntroTutorials"));
 
 	// Everything (or almost) uses this padding, change it to expand the padding.
@@ -346,8 +271,6 @@ void SFoliageEdit::Construct(const FArguments& InArgs)
 					]
 
 					+ SVerticalBox::Slot()
-					.AutoHeight()
-					.HAlign(HAlign_Fill)
 					.Padding(StandardPadding)
 					[
 						SNew(SWrapBox)
@@ -356,69 +279,68 @@ void SFoliageEdit::Construct(const FArguments& InArgs)
 
 						// Select all instances
 						+ SWrapBox::Slot()
-						.Padding(FMargin(0.f, 0.f, 5.f, 0.f))
+						.Padding(FMargin(0.f, 0.f, 6.f, 3.f))
 						[
-							SNew(SButton)
-							.OnClicked(this, &SFoliageEdit::OnSelectAllInstances)
-							.Text(LOCTEXT("SelectAllInstances", "Select All"))
-							.ToolTipText(LOCTEXT("SelectAllInstances_Tooltip", "Selects all foliage instances"))
+							SNew(SBox)
+							.WidthOverride(100.f)
+							.HeightOverride(25.f)
+							[
+								SNew(SButton)
+								.HAlign(HAlign_Center)
+								.VAlign(VAlign_Center)
+								.OnClicked(this, &SFoliageEdit::OnSelectAllInstances)
+								.Text(LOCTEXT("SelectAllInstances", "Select All"))
+								.ToolTipText(LOCTEXT("SelectAllInstances_Tooltip", "Selects all foliage instances"))
+							]
 						]
 
 						// Select all invalid instances
 						+ SWrapBox::Slot()
-						.Padding(FMargin(0.f, 0.f, 5.f, 0.f))
+						.Padding(FMargin(0.f, 0.f, 6.f, 3.f))
 						[
-							SNew(SButton)
-							.OnClicked(this, &SFoliageEdit::OnSelectInvalidInstances)
-							.Text(LOCTEXT("SelectInvalidInstances", "Select Invalid"))
-							.ToolTipText(LOCTEXT("SelectInvalidInstances_Tooltip", "Selects all foliage instances that are not placed in a valid location"))
+							SNew(SBox)
+							.WidthOverride(100.f)
+							.HeightOverride(25.f)
+							[
+								SNew(SButton)
+								.HAlign(HAlign_Center)
+								.VAlign(VAlign_Center)
+								.OnClicked(this, &SFoliageEdit::OnSelectInvalidInstances)
+								.Text(LOCTEXT("SelectInvalidInstances", "Select Invalid"))
+								.ToolTipText(LOCTEXT("SelectInvalidInstances_Tooltip", "Selects all foliage instances that are not placed in a valid location"))
+						
+							]
 						]
 
 						// Deselect all
 						+ SWrapBox::Slot()
-						.Padding(FMargin(0.f, 0.f, 5.f, 0.f))
+						.Padding(FMargin(0.f, 0.f, 6.f, 3.f))
 						[
-							SNew(SButton)
-							.OnClicked(this, &SFoliageEdit::OnDeselectAllInstances)
-							.Text(LOCTEXT("DeselectAllInstances", "Deselect All"))
-							.ToolTipText(LOCTEXT("DeselectAllInstances_Tooltip", "Deselects all foliage instances"))
+							SNew(SBox)
+							.WidthOverride(100.f)
+							.HeightOverride(25.f)
+							[
+								SNew(SButton)
+								.HAlign(HAlign_Center)
+								.VAlign(VAlign_Center)
+								.OnClicked(this, &SFoliageEdit::OnDeselectAllInstances)
+								.Text(LOCTEXT("DeselectAllInstances", "Deselect All"))
+								.ToolTipText(LOCTEXT("DeselectAllInstances_Tooltip", "Deselects all foliage instances"))
+							]
 						]
 					]
 				]
 			]
 		]
 
+		// Foliage Palette
 		+ SVerticalBox::Slot()
 		.FillHeight(1.f)
 		.VAlign(VAlign_Fill)
 		.Padding(0.f, 5.f, 0.f, 0.f)
 		[
-			SNew(SOverlay)
-
-			+ SOverlay::Slot()
-			[
-				SAssignNew(FoliagePalette, SFoliagePalette)
-				.FoliageEdMode(FoliageEditMode)
-			]
-			
-			// Foliage Mesh Drop Zone
-			+ SOverlay::Slot()
-			.HAlign(HAlign_Fill)
-			.VAlign(VAlign_Fill)
-			[
-				SNew(SFoliageDragDropHandler)
-				.Visibility(this, &SFoliageEdit::GetVisibility_FoliageDropTarget)
-				.FoliageEditPtr(SharedThis(this))
-				[
-					SNew(SScaleBox)
-					.Stretch(EStretch::ScaleToFit)
-					[
-						SNew(STextBlock)
-						.Text(LOCTEXT("Foliage_AddFoliageMesh", "+ Foliage Mesh"))
-						.ShadowOffset(FVector2D(1.f, 1.f))
-					]
-				]
-			]
+			SAssignNew(FoliagePalette, SFoliagePalette)
+			.FoliageEdMode(FoliageEditMode)
 		]
 	];
 
@@ -428,56 +350,12 @@ END_SLATE_FUNCTION_BUILD_OPTIMIZATION
 
 void SFoliageEdit::RefreshFullList()
 {
-	FoliagePalette->Refresh();
-}
-
-void SFoliageEdit::ClearAllToolSelection()
-{
-	FoliageEditMode->UISettings.SetLassoSelectToolSelected(false);
-	FoliageEditMode->UISettings.SetPaintToolSelected(false);
-	FoliageEditMode->UISettings.SetReapplyToolSelected(false);
-	FoliageEditMode->UISettings.SetSelectToolSelected(false);
-	FoliageEditMode->UISettings.SetPaintBucketToolSelected(false);
-}
-
-void SFoliageEdit::BindCommands()
-{
-	const FFoliageEditCommands& Commands = FFoliageEditCommands::Get();
-
-	UICommandList->MapAction(
-		Commands.SetPaint,
-		FExecuteAction::CreateSP(this, &SFoliageEdit::OnSetPaint),
-		FCanExecuteAction(),
-		FIsActionChecked::CreateSP(this, &SFoliageEdit::IsPaintTool));
-
-	UICommandList->MapAction(
-		Commands.SetReapplySettings,
-		FExecuteAction::CreateSP(this, &SFoliageEdit::OnSetReapplySettings),
-		FCanExecuteAction(),
-		FIsActionChecked::CreateSP(this, &SFoliageEdit::IsReapplySettingsTool));
-
-	UICommandList->MapAction(
-		Commands.SetSelect,
-		FExecuteAction::CreateSP(this, &SFoliageEdit::OnSetSelectInstance),
-		FCanExecuteAction(),
-		FIsActionChecked::CreateSP(this, &SFoliageEdit::IsSelectTool));
-
-	UICommandList->MapAction(
-		Commands.SetLassoSelect,
-		FExecuteAction::CreateSP(this, &SFoliageEdit::OnSetLasso),
-		FCanExecuteAction(),
-		FIsActionChecked::CreateSP(this, &SFoliageEdit::IsLassoSelectTool));
-
-	UICommandList->MapAction(
-		Commands.SetPaintBucket,
-		FExecuteAction::CreateSP(this, &SFoliageEdit::OnSetPaintFill),
-		FCanExecuteAction(),
-		FIsActionChecked::CreateSP(this, &SFoliageEdit::IsPaintFillTool));
+	FoliagePalette->UpdatePalette(true);
 }
 
 TSharedRef<SWidget> SFoliageEdit::BuildToolBar()
 {
-	FToolBarBuilder Toolbar(UICommandList, FMultiBoxCustomization::None, nullptr, Orient_Vertical);
+	FToolBarBuilder Toolbar(FoliageEditMode->UICommandList, FMultiBoxCustomization::None, nullptr, Orient_Vertical);
 	Toolbar.SetLabelVisibility(EVisibility::Collapsed);
 	Toolbar.SetStyle(&FEditorStyle::Get(), "FoliageEditToolbar");
 	{
@@ -508,84 +386,26 @@ TSharedRef<SWidget> SFoliageEdit::BuildToolBar()
 		];
 }
 
-void SFoliageEdit::AddFoliageType(const FAssetData& AssetData)
-{
-	GWarn->BeginSlowTask(LOCTEXT("AddFoliageType_LoadPackage", "Fully Loading Package for Add"), true, false);
-	UObject* Asset = AssetData.GetAsset();
-	GWarn->EndSlowTask();
-
-	UFoliageType* FoliageType = FoliageEditMode->AddFoliageAsset(Asset);
-	if (FoliageType)
-	{
-		FoliagePalette->Refresh();
-	}
-}
-
-void SFoliageEdit::OnSetPaint()
-{
-	ClearAllToolSelection();
-
-	FoliageEditMode->UISettings.SetPaintToolSelected(true);
-
-	FoliageEditMode->NotifyToolChanged();
-}
-
 bool SFoliageEdit::IsPaintTool() const
 {
 	return FoliageEditMode->UISettings.GetPaintToolSelected();
 }
 
-void SFoliageEdit::OnSetReapplySettings()
-{
-	ClearAllToolSelection();
-
-	FoliageEditMode->UISettings.SetReapplyToolSelected(true);
-
-	FoliageEditMode->NotifyToolChanged();
-}
-
 bool SFoliageEdit::IsReapplySettingsTool() const
 {
 	return FoliageEditMode->UISettings.GetReapplyToolSelected();
 }
 
-void SFoliageEdit::OnSetSelectInstance()
-{
-	ClearAllToolSelection();
-
-	FoliageEditMode->UISettings.SetSelectToolSelected(true);
-
-	FoliageEditMode->NotifyToolChanged();
-}
-
 bool SFoliageEdit::IsSelectTool() const
 {
 	return FoliageEditMode->UISettings.GetSelectToolSelected();
 }
 
-void SFoliageEdit::OnSetLasso()
-{
-	ClearAllToolSelection();
-
-	FoliageEditMode->UISettings.SetLassoSelectToolSelected(true);
-
-	FoliageEditMode->NotifyToolChanged();
-}
-
 bool SFoliageEdit::IsLassoSelectTool() const
 {
 	return FoliageEditMode->UISettings.GetLassoSelectToolSelected();
 }
 
-void SFoliageEdit::OnSetPaintFill()
-{
-	ClearAllToolSelection();
-
-	FoliageEditMode->UISettings.SetPaintBucketToolSelected(true);
-
-	FoliageEditMode->NotifyToolChanged();
-}
-
 bool SFoliageEdit::IsPaintFillTool() const
 {
 	return FoliageEditMode->UISettings.GetPaintBucketToolSelected();
@@ -654,23 +474,6 @@ EVisibility SFoliageEdit::GetVisibility_SelectionOptions() const
 	return IsSelectTool() || IsLassoSelectTool() ? EVisibility::SelfHitTestInvisible : EVisibility::Collapsed;
 }
 
-EVisibility SFoliageEdit::GetVisibility_FoliageDropTarget() const
-{
-	if ( FSlateApplication::Get().IsDragDropping() )
-	{
-		TArray<FAssetData> DraggedAssets = AssetUtil::ExtractAssetDataFromDrag(FSlateApplication::Get().GetDragDroppingContent());
-		for ( const FAssetData& AssetData : DraggedAssets )
-		{
-			if ( AssetData.IsValid() && (AssetData.GetClass()->IsChildOf(UStaticMesh::StaticClass()) ||	AssetData.GetClass()->IsChildOf(UFoliageType::StaticClass())) )
-			{
-				return EVisibility::Visible;
-			}
-		}
-	}
-	
-	return EVisibility::Hidden;
-}
-
 FReply SFoliageEdit::OnSelectAllInstances()
 {
 	for (FFoliageMeshUIInfoPtr& TypeInfo : FoliageEditMode->GetFoliageMeshList())
@@ -714,22 +517,6 @@ FReply SFoliageEdit::OnDeselectAllInstances()
 	return FReply::Handled();
 }
 
-FReply SFoliageEdit::OnDrop_ListView(const FGeometry& MyGeometry, const FDragDropEvent& DragDropEvent)
-{
-	TArray<FAssetData> DroppedAssetData = AssetUtil::ExtractAssetDataFromDrag(DragDropEvent);
-	{
-		// Treat the entire drop as a transaction (in case multiples types are being added)
-		const FScopedTransaction Transaction(NSLOCTEXT("UnrealEd", "FoliageMode_DragDropTypesTransaction", "Drag-drop Foliage Type(s)"));
-
-		for (int32 AssetIdx = 0; AssetIdx < DroppedAssetData.Num(); ++AssetIdx)
-		{	
-			AddFoliageType(DroppedAssetData[AssetIdx]);
-		}
-	}
-
-	return FReply::Handled();
-}
-
 FText SFoliageEdit::GetFilterText() const
 {
 	FText TooltipText;
diff --git a/Engine/Source/Editor/FoliageEdit/Private/SFoliageEdit.h b/Engine/Source/Editor/FoliageEdit/Private/SFoliageEdit.h
--- a/Engine/Source/Editor/FoliageEdit/Private/SFoliageEdit.h
+++ b/Engine/Source/Editor/FoliageEdit/Private/SFoliageEdit.h
@@ -18,54 +18,24 @@ public:
 	/** Does a full refresh on the list. */
 	void RefreshFullList();
 
-	/** Handles adding a new item to the list and refreshing the list in it's entirety. */
-	FReply OnDrop_ListView(const FGeometry& MyGeometry, const FDragDropEvent& DragDropEvent);
-
 	/** Gets FoliageEditMode. Used by the cluster details to notify changes */
 	class FEdModeFoliage* GetFoliageEditMode() const { return FoliageEditMode; }
 
 private:
-	/** Clears all the tools selection by setting them to false. */
-	void ClearAllToolSelection();
-
-	/** Binds UI commands for the toolbar. */
-	void BindCommands();
-
 	/** Creates the toolbar. */
 	TSharedRef<SWidget> BuildToolBar();
 
-	/** Adds the foliage type asset to the foliage actor's list of meshes. */
-	void AddFoliageType(const class FAssetData& AssetData);
-
-	/** Delegate callbacks for the UI */
-
-	/** Sets the tool mode to Paint. */
-	void OnSetPaint();
-
 	/** Checks if the tool mode is Paint. */
 	bool IsPaintTool() const;
 
-	/** Sets the tool mode to Reapply Settings. */
-	void OnSetReapplySettings();
-
 	/** Checks if the tool mode is Reapply Settings. */
 	bool IsReapplySettingsTool() const;
 
-	/** Sets the tool mode to Select. */
-	void OnSetSelectInstance();
-
 	/** Checks if the tool mode is Select. */
 	bool IsSelectTool() const;
 
-	/** Sets the tool mode to Lasso Select. */
-	void OnSetLasso();
-
-	/** Checks if the tool mode is Lasso Select. */
 	bool IsLassoSelectTool() const;
 
-	/** Sets the tool mode to Paint Bucket. */
-	void OnSetPaintFill();
-
 	/** Checks if the tool mode is Paint Bucket. */
 	bool IsPaintFillTool() const;
 
@@ -144,9 +114,6 @@ private:	// BRUSH SETTINGS
 	/** Checks if the text in the empty list overlay should appear. If the list is has items but the the drag and drop override is true, it will return EVisibility::Visible. */
 	EVisibility GetVisibility_EmptyList() const;
 
-	/** Should the drop area be visible currently?  Happens when the user is dragging static meshes */
-	EVisibility GetVisibility_FoliageDropTarget() const;
-
 	/** @return Whether selecting instances is currently possible */
 	EVisibility GetVisibility_SelectionOptions() const;
 
@@ -172,9 +139,6 @@ private:
 	/** Palette of available foliage types */
 	TSharedPtr<class SFoliagePalette> FoliagePalette;
 
-	/** Command list for binding functions for the toolbar. */
-	TSharedPtr<FUICommandList>		UICommandList;
-
 	/** Pointer to the foliage edit mode. */
 	FEdModeFoliage*					FoliageEditMode;
 };
diff --git a/Engine/Source/Editor/FoliageEdit/Private/SFoliagePalette.cpp b/Engine/Source/Editor/FoliageEdit/Private/SFoliagePalette.cpp
--- a/Engine/Source/Editor/FoliageEdit/Private/SFoliagePalette.cpp
+++ b/Engine/Source/Editor/FoliageEdit/Private/SFoliagePalette.cpp
@@ -2,7 +2,7 @@
 #include "UnrealEd.h"
 #include "FoliageType_InstancedStaticMesh.h"
 #include "FoliageEdMode.h"
-#include "FoliageEditActions.h"
+#include "FoliagePaletteCommands.h"
 #include "SFoliagePalette.h"
 #include "FoliageTypePaintingCustomization.h"
 #include "FoliagePaletteItem.h"
@@ -15,98 +15,108 @@
 #include "Editor/UnrealEd/Public/AssetSelection.h"
 #include "Editor/PropertyEditor/Public/IDetailsView.h"
 #include "Editor/PropertyEditor/Public/PropertyEditorModule.h"
+#include "Editor/UnrealEd/Public/AssetSelection.h"
+#include "Editor/UnrealEd/Public/ScopedTransaction.h"
+
 #include "Engine/StaticMesh.h"
 
 #include "SScaleBox.h"
 #include "SWidgetSwitcher.h"
-#include "SExpandableArea.h"
+#include "SSearchBox.h"
 
 #define LOCTEXT_NAMESPACE "FoliageEd_Mode"
 
+////////////////////////////////////////////////
+// SFoliageDragDropHandler
+////////////////////////////////////////////////
+
+/** Drag-drop zone for adding foliage types to the palette */
+class SFoliageDragDropHandler : public SCompoundWidget
+{
+public:
+	SLATE_BEGIN_ARGS(SFoliageDragDropHandler) {}
+		SLATE_DEFAULT_SLOT(FArguments, Content)
+		SLATE_EVENT(FOnDrop, OnDrop)
+	SLATE_END_ARGS()
+
+	void Construct(const FArguments& InArgs)
+	{
+		bIsDragOn = false;
+		OnDropDelegate = InArgs._OnDrop;
+
+		this->ChildSlot
+			[
+				SNew(SBorder)
+				.BorderImage(FEditorStyle::GetBrush("WhiteBrush"))
+				.BorderBackgroundColor(this, &SFoliageDragDropHandler::GetBackgroundColor)
+				.Padding(FMargin(30))
+				[
+					InArgs._Content.Widget
+				]
+			];
+	}
+
+	FReply OnDrop(const FGeometry& MyGeometry, const FDragDropEvent& DragDropEvent) override
+	{
+		bIsDragOn = false;
+		if (OnDropDelegate.IsBound())
+		{
+			return OnDropDelegate.Execute(MyGeometry, DragDropEvent);
+		}
+		
+		return FReply::Handled();
+	}
+
+	virtual void OnDragEnter(const FGeometry& MyGeometry, const FDragDropEvent& DragDropEvent) override
+	{
+		bIsDragOn = true;
+	}
+
+	virtual void OnDragLeave(const FDragDropEvent& DragDropEvent) override
+	{
+		bIsDragOn = false;
+	}
+
+private:
+	FSlateColor GetBackgroundColor() const
+	{
+		return bIsDragOn ? FLinearColor(1.0f, 0.6f, 0.1f, 0.9f) : FLinearColor(0.1f, 0.1f, 0.1f, 0.9f);
+	}
+
+private:
+	FOnDrop OnDropDelegate;
+	bool bIsDragOn;
+};
+
+////////////////////////////////////////////////
+// SFoliagePalette
+////////////////////////////////////////////////
 void SFoliagePalette::Construct(const FArguments& InArgs)
 {
 	FoliageEditMode = InArgs._FoliageEdMode;
 
-	//@todo: Save these as config options
-	bShowTooltips = true;
-	ActiveViewMode = EPaletteViewMode::Thumbnail;
-
 	FoliageEditMode->OnToolChanged.AddSP(this, &SFoliagePalette::HandleOnToolChanged);
 
+	FFoliagePaletteCommands::Register();
+	UICommandList = MakeShareable(new FUICommandList);
 	BindCommands();
 
+	TypeFilter = MakeShareable(new FoliageTypeTextFilter(
+		FoliageTypeTextFilter::FItemToStringArray::CreateSP(this, &SFoliagePalette::GetPaletteItemFilterString)));
+
 	FPropertyEditorModule& PropertyModule = FModuleManager::LoadModuleChecked<FPropertyEditorModule>("PropertyEditor");
 	FDetailsViewArgs Args(false, false, false, FDetailsViewArgs::HideNameArea, true);
 	Args.bShowActorLabel = false;
 	DetailsWidget = PropertyModule.CreateDetailView(Args);
+	DetailsWidget->SetVisibility(FoliageEditMode->UISettings.GetShowPaletteItemDetails() ? EVisibility::SelfHitTestInvisible : EVisibility::Collapsed);
 
 	// We want to use our own customization for UFoliageType
 	DetailsWidget->RegisterInstancedCustomPropertyLayout(UFoliageType::StaticClass(), 
 		FOnGetDetailCustomizationInstance::CreateStatic(&FFoliageTypePaintingCustomization::MakeInstance, FoliageEditMode)
 		);
 
-	FMargin StandardPadding(6.f, 3.f);
 	const FText BlankText = LOCTEXT("Blank", "");
 
-	// Tile View Widget
-	static const float TileSize = 64.f;
-	SAssignNew(TileViewWidget, SFoliageTypeTileView)
-		.ListItemsSource(&FoliageEditMode->GetFoliageMeshList())
-		.SelectionMode(ESelectionMode::Multi)
-		.OnGenerateTile(this, &SFoliagePalette::GenerateTile)
-		.OnContextMenuOpening(this, &SFoliagePalette::ConstructFoliageTypeContextMenu)
-		.OnSelectionChanged(this, &SFoliagePalette::OnSelectionChanged)
-		.ItemHeight(TileSize)
-		.ItemWidth(TileSize)
-		.ItemAlignment(EListItemAlignment::LeftAligned)
-		.ClearSelectionOnClick(true)
-		.OnMouseButtonDoubleClick(this, &SFoliagePalette::OnItemDoubleClicked);
-
-	// Tree View Widget
-	SAssignNew(TreeViewWidget, SFoliageTypeTreeView)
-	.TreeItemsSource(&FoliageEditMode->GetFoliageMeshList())
-	.SelectionMode(ESelectionMode::Multi)
-	.OnGenerateRow(this, &SFoliagePalette::TreeViewGenerateRow)
-	.OnGetChildren(this, &SFoliagePalette::TreeViewGetChildren)
-	.OnContextMenuOpening(this, &SFoliagePalette::ConstructFoliageTypeContextMenu)
-	.OnSelectionChanged(this, &SFoliagePalette::OnSelectionChanged)
-	.OnMouseButtonDoubleClick(this, &SFoliagePalette::OnItemDoubleClicked)
-	.HeaderRow
-	(
-		// Toggle Active
-		SAssignNew(TreeViewHeaderRow, SHeaderRow)
-		+ SHeaderRow::Column(FoliagePaletteTreeColumns::ColumnID_ToggleActive)
-		[
-			SNew(SCheckBox)
-			.IsChecked(this, &SFoliagePalette::GetState_AllMeshes)
-			.OnCheckStateChanged(this, &SFoliagePalette::OnCheckStateChanged_AllMeshes)
-		]
-		.DefaultLabel(BlankText)
-		.HeaderContentPadding(FMargin(0, 1, 0, 1))
-		.HAlignHeader(HAlign_Center)
-		.HAlignCell(HAlign_Center)
-		.FixedWidth(24)
-
-		// Type
-		+ SHeaderRow::Column(FoliagePaletteTreeColumns::ColumnID_Type)
-		.HeaderContentPadding(FMargin(10, 1, 0, 1))
-		.SortMode(this, &SFoliagePalette::GetMeshColumnSortMode)
-		.OnSort(this, &SFoliagePalette::OnMeshesColumnSortModeChanged)
-		.DefaultLabel(this, &SFoliagePalette::GetMeshesHeaderText)
-
-		// Instance Count
-		+ SHeaderRow::Column(FoliagePaletteTreeColumns::ColumnID_InstanceCount)
-		.HeaderContentPadding(FMargin(10, 1, 0, 1))
-		.DefaultLabel(LOCTEXT("InstanceCount", "Count"))
-		.DefaultTooltip(this, &SFoliagePalette::GetTotalInstanceCountTooltipText)
-		.FixedWidth(60.f)
-
-		// Save Asset
-		+ SHeaderRow::Column(FoliagePaletteTreeColumns::ColumnID_Save)
-		.FixedWidth(24.0f)
-		.DefaultLabel(BlankText)
-	);
-
 	ChildSlot
 	[
 		SNew(SVerticalBox)
@@ -158,21 +168,14 @@ void SFoliagePalette::Construct(const FArguments& InArgs)
 				]
 
 				+ SHorizontalBox::Slot()
-				.HAlign(HAlign_Right)
+				.HAlign(HAlign_Fill)
+				.VAlign(VAlign_Center)
+				.Padding(6.f, 0.f)
 				[
-					// View Options
-					SNew( SComboButton )
-					.ContentPadding(0)
-					.ForegroundColor( FSlateColor::UseForeground() )
-					.ButtonStyle( FEditorStyle::Get(), "ToggleButton" )
-					.OnGetMenuContent(this, &SFoliagePalette::GetViewOptionsMenuContent)
-					.ButtonContent()
-					[
-						SNew(SImage)
-						.Image( FEditorStyle::GetBrush("GenericViewButton") )
-					]
+					SNew(SSearchBox)
+					.HintText(LOCTEXT("SearchFoliagePaletteHint", "Search Foliage"))
+					.OnTextChanged(this, &SFoliagePalette::OnSearchTextChanged)
 				]
-				
 			]
 		]
 
@@ -185,105 +188,128 @@ void SFoliagePalette::Construct(const FArguments& InArgs)
 			+ SSplitter::Slot()
 			.Value(0.6f)
 			[
-				SNew(SVerticalBox)
+				SNew(SOverlay)
 
-				+ SVerticalBox::Slot()
-				.AutoHeight()
-				.Padding(StandardPadding)
+				+ SOverlay::Slot()
 				[
-					SNew(SBox)
-					.Visibility(this, &SFoliagePalette::GetDropFoliageHintVisibility)
-					.Padding(FMargin(15, 0))
-					.MinDesiredHeight(30)
-					[
-						SNew(SScaleBox)
-						.Stretch(EStretch::ScaleToFit)
-						[
-							SNew(STextBlock)
-							.Text(LOCTEXT("Foliage_DropStatic", "+ Drop Foliage Here"))
-							.ToolTipText(LOCTEXT("Foliage_DropStatic_ToolTip", "Drag and drop foliage types or static meshes from the Content Browser to add them to the palette"))
-						]
-					]
-				]
-
-				+ SVerticalBox::Slot()
-				[
-					SAssignNew(WidgetSwitcher, SWidgetSwitcher)
+					SNew(SVerticalBox)
 
-					+ SWidgetSwitcher::Slot()
+					+ SVerticalBox::Slot()
+					.AutoHeight()
+					.Padding(FMargin(6.f, 3.f))
 					[
-						// Thumbnail View
-						SNew(SScrollBorder, TileViewWidget.ToSharedRef())
-						.Content()
+						SNew(SBox)
+						.Visibility(this, &SFoliagePalette::GetDropFoliageHintVisibility)
+						.Padding(FMargin(15, 0))
+						.MinDesiredHeight(30)
 						[
-							TileViewWidget.ToSharedRef()
+							SNew(SScaleBox)
+							.Stretch(EStretch::ScaleToFit)
+							[
+								SNew(STextBlock)
+								.Text(LOCTEXT("Foliage_DropStatic", "+ Drop Foliage Here"))
+								.ToolTipText(LOCTEXT("Foliage_DropStatic_ToolTip", "Drag and drop foliage types or static meshes from the Content Browser to add them to the palette"))
+							]
 						]
 					]
 
-					+ SWidgetSwitcher::Slot()
+					+ SVerticalBox::Slot()
 					[
-						// Tree View
-						SNew(SScrollBorder, TreeViewWidget.ToSharedRef())
-						.Style(&FEditorStyle::Get().GetWidgetStyle<FScrollBorderStyle>("FoliageEditMode.TreeView.ScrollBorder"))
-						.Content()
-						[
-							TreeViewWidget.ToSharedRef()
-						]
+						CreatePaletteViews()
 					]
-				]
 				
-				+ SVerticalBox::Slot()
-				.Padding(FMargin(0.f))
-				.VAlign(VAlign_Bottom)
-				.AutoHeight()
-				[
-					SNew(SHorizontalBox)
-					
-					// Selected type name area
-					+ SHorizontalBox::Slot()
-					.Padding(FMargin(3.f))
+					+ SVerticalBox::Slot()
+					.Padding(FMargin(0.f))
 					.VAlign(VAlign_Bottom)
-					.AutoWidth()
+					.AutoHeight()
 					[
-						SNew(STextBlock)
-						.Text(this, &SFoliagePalette::GetDetailsNameAreaText)
-					]
-
-					// Show/Hide Details
-					+ SHorizontalBox::Slot()
-					.HAlign(HAlign_Right)
-					[
-						SNew(SButton)
-						.ToolTipText(this, &SFoliagePalette::GetShowHideDetailsTooltipText)
-						.ForegroundColor(FSlateColor::UseForeground())
-						.ButtonStyle(FEditorStyle::Get(), "ToggleButton")
-						.OnClicked(this, &SFoliagePalette::OnShowHideDetailsClicked)
-						.ContentPadding(FMargin(2.f))
-						.Content()
+						SNew(SHorizontalBox)
+					
+						// Show/Hide Details
+						+ SHorizontalBox::Slot()
+						.HAlign(HAlign_Right)
+						.AutoWidth()
 						[
-							SNew(SHorizontalBox)
-
-							+SHorizontalBox::Slot()
-							.AutoWidth()
-							.HAlign(HAlign_Center)
-							.VAlign(VAlign_Center)
+							SNew(SButton)
+							.ToolTipText(this, &SFoliagePalette::GetShowHideDetailsTooltipText)
+							.ForegroundColor(FSlateColor::UseForeground())
+							.ButtonStyle(FEditorStyle::Get(), "ToggleButton")
+							.OnClicked(this, &SFoliagePalette::OnShowHideDetailsClicked)
+							.ContentPadding(FMargin(2.f))
+							.Content()
 							[
-								SNew(SImage)
-								.Image(FEditorStyle::GetBrush("LevelEditor.Tabs.Details"))
+								SNew(SHorizontalBox)
+
+								// Arrow
+								+ SHorizontalBox::Slot()
+								.Padding(FMargin(3.f, 0.f))
+								.AutoWidth()
+								.HAlign(HAlign_Center)
+								.VAlign(VAlign_Center)
+								[
+									SNew(SImage)
+									.Image(this, &SFoliagePalette::GetShowHideDetailsImage)
+								]
+
+								// Details icon
+								+SHorizontalBox::Slot()
+								.AutoWidth()
+								.HAlign(HAlign_Center)
+								.VAlign(VAlign_Center)
+								[
+									SNew(SImage)
+									.Image(FEditorStyle::GetBrush("LevelEditor.Tabs.Details"))
+								]
 							]
+						]
+
+						// Selected type name area
+						+ SHorizontalBox::Slot()
+						.Padding(FMargin(3.f))
+						.VAlign(VAlign_Bottom)
+						//.AutoWidth()
+						[
+							SNew(STextBlock)
+							.Text(this, &SFoliagePalette::GetDetailsNameAreaText)
+						]
 
-							+ SHorizontalBox::Slot()
-							.Padding(FMargin(3.f, 0.f))
-							.AutoWidth()
-							.HAlign(HAlign_Center)
-							.VAlign(VAlign_Center)
+						// View Options
+						+ SHorizontalBox::Slot()
+						.HAlign(HAlign_Right)
+						.AutoWidth()
+						[
+							SNew( SComboButton )
+							.ContentPadding(0)
+							.ForegroundColor( FSlateColor::UseForeground() )
+							.ButtonStyle( FEditorStyle::Get(), "ToggleButton" )
+							.OnGetMenuContent(this, &SFoliagePalette::GetViewOptionsMenuContent)
+							.ButtonContent()
 							[
 								SNew(SImage)
-								.Image(this, &SFoliagePalette::GetShowHideDetailsImage)
+								.Image( FEditorStyle::GetBrush("GenericViewButton") )
 							]
 						]
 					]
 				]
+				
+				// Foliage Mesh Drop Zone
+				+ SOverlay::Slot()
+				.HAlign(HAlign_Fill)
+				.VAlign(VAlign_Fill)
+				[
+					SNew(SFoliageDragDropHandler)
+					.Visibility(this, &SFoliagePalette::GetFoliageDropTargetVisibility)
+					.OnDrop(this, &SFoliagePalette::HandleFoliageDropped)
+					[
+						SNew(SScaleBox)
+						.Stretch(EStretch::ScaleToFit)
+						[
+							SNew(STextBlock)
+							.Text(LOCTEXT("Foliage_AddFoliageMesh", "+ Foliage Type"))
+							.ShadowOffset(FVector2D(1.f, 1.f))
+						]
+					]
+				]
 			]
 
 			// Details
@@ -293,12 +319,29 @@ void SFoliagePalette::Construct(const FArguments& InArgs)
 			]
 		]
 	];
+
+	UpdatePalette(true);
+}
+
+FReply SFoliagePalette::OnKeyDown(const FGeometry& MyGeometry, const FKeyEvent& InKeyEvent)
+{
+	if (UICommandList->ProcessCommandBindings(InKeyEvent))
+	{
+		return FReply::Handled();
+	}
+
+	return SCompoundWidget::OnKeyDown(MyGeometry, InKeyEvent);
 }
 
-void SFoliagePalette::Refresh()
+void SFoliagePalette::UpdatePalette(bool bRefreshItems)
 {
-	TreeViewWidget->RequestTreeRefresh();
-	TileViewWidget->RequestListRefresh();
+	bItemsNeedRefresh |= bRefreshItems;
+
+	if (!bIsActiveTimerRegistered)
+	{
+		bIsActiveTimerRegistered = true;
+		RegisterActiveTimer(0.f, FWidgetActiveTimerDelegate::CreateSP(this, &SFoliagePalette::UpdatePaletteItems));
+	}
 }
 
 bool SFoliagePalette::AnySelectedTileHovered() const
@@ -317,31 +360,24 @@ bool SFoliagePalette::AnySelectedTileHovered() const
 void SFoliagePalette::ActivateAllSelectedTypes(bool bActivate) const
 {
 	// Apply the new check state to all of the selected types
-	for (FFoliageMeshUIInfoPtr& TypeInfo : GetActiveViewWidget()->GetSelectedItems())
+	for (FFoliagePaletteItemModelPtr& PaletteItem : GetActiveViewWidget()->GetSelectedItems())
 	{
-		if (TypeInfo->Settings->IsSelected != bActivate)
-		{
-			TypeInfo->Settings->Modify();
-			TypeInfo->Settings->IsSelected = bActivate;
-		}
+		PaletteItem->SetTypeActiveInPalette(bActivate);
 	}
 }
 
 void SFoliagePalette::BindCommands()
 {
-	UICommandList = MakeShareable(new FUICommandList);
-	const FFoliageEditCommands& Commands = FFoliageEditCommands::Get();
+	const FFoliagePaletteCommands& Commands = FFoliagePaletteCommands::Get();
 
 	// Context menu commands
 	UICommandList->MapAction(
 		Commands.RemoveFoliageType,
-		FExecuteAction::CreateSP(this, &SFoliagePalette::OnRemoveFoliageType),
-		FCanExecuteAction());
+		FExecuteAction::CreateSP(this, &SFoliagePalette::OnRemoveFoliageType));
 
 	UICommandList->MapAction(
 		Commands.ShowFoliageTypeInCB,
-		FExecuteAction::CreateSP(this, &SFoliagePalette::OnShowFoliageTypeInCB),
-		FCanExecuteAction());
+		FExecuteAction::CreateSP(this, &SFoliagePalette::OnShowFoliageTypeInCB));
 
 	UICommandList->MapAction(
 		Commands.SelectAllInstances,
@@ -359,6 +395,18 @@ void SFoliagePalette::BindCommands()
 		FCanExecuteAction::CreateSP(this, &SFoliagePalette::CanSelectInstances));
 }
 
+void SFoliagePalette::RefreshActivePaletteViewWidget()
+{
+	if (FoliageEditMode->UISettings.GetActivePaletteViewMode() == EFoliagePaletteViewMode::Thumbnail)
+	{
+		TileViewWidget->RequestListRefresh();
+	}
+	else
+	{
+		TreeViewWidget->RequestTreeRefresh();
+	}
+}
+
 void SFoliagePalette::AddFoliageType(const FAssetData& AssetData)
 {
 	if (AddFoliageTypeCombo.IsValid())
@@ -373,10 +421,112 @@ void SFoliagePalette::AddFoliageType(const FAssetData& AssetData)
 	UFoliageType* FoliageType = FoliageEditMode->AddFoliageAsset(Asset);
 	if (FoliageType)
 	{
-		Refresh();
+		UpdatePalette(true);
 	}
 }
 
+TSharedRef<SWidgetSwitcher> SFoliagePalette::CreatePaletteViews()
+{
+	const FText BlankText = LOCTEXT("Blank", "");
+
+	// Tile View Widget
+	SAssignNew(TileViewWidget, SFoliageTypeTileView)
+		.ListItemsSource(&FilteredItems)
+		.SelectionMode(ESelectionMode::Multi)
+		.OnGenerateTile(this, &SFoliagePalette::GenerateTile)
+		.OnContextMenuOpening(this, &SFoliagePalette::ConstructFoliageTypeContextMenu)
+		.OnSelectionChanged(this, &SFoliagePalette::OnSelectionChanged)
+		.ItemHeight(this, &SFoliagePalette::GetScaledThumbnailSize)
+		.ItemWidth(this, &SFoliagePalette::GetScaledThumbnailSize)
+		.ItemAlignment(EListItemAlignment::LeftAligned)
+		.ClearSelectionOnClick(true)
+		.OnMouseButtonDoubleClick(this, &SFoliagePalette::OnItemDoubleClicked);
+
+	// Tree View Widget
+	SAssignNew(TreeViewWidget, SFoliageTypeTreeView)
+	.TreeItemsSource(&FilteredItems)
+	.SelectionMode(ESelectionMode::Multi)
+	.OnGenerateRow(this, &SFoliagePalette::TreeViewGenerateRow)
+	.OnGetChildren(this, &SFoliagePalette::TreeViewGetChildren)
+	.OnContextMenuOpening(this, &SFoliagePalette::ConstructFoliageTypeContextMenu)
+	.OnSelectionChanged(this, &SFoliagePalette::OnSelectionChanged)
+	.OnMouseButtonDoubleClick(this, &SFoliagePalette::OnItemDoubleClicked)
+	.HeaderRow
+	(
+		// Toggle Active
+		SAssignNew(TreeViewHeaderRow, SHeaderRow)
+		+ SHeaderRow::Column(FoliagePaletteTreeColumns::ColumnID_ToggleActive)
+		[
+			SNew(SCheckBox)
+			.IsChecked(this, &SFoliagePalette::GetState_AllMeshes)
+			.OnCheckStateChanged(this, &SFoliagePalette::OnCheckStateChanged_AllMeshes)
+		]
+		.DefaultLabel(BlankText)
+		.HeaderContentPadding(FMargin(0, 1, 0, 1))
+		.HAlignHeader(HAlign_Center)
+		.HAlignCell(HAlign_Center)
+		.FixedWidth(24)
+
+		// Type
+		+ SHeaderRow::Column(FoliagePaletteTreeColumns::ColumnID_Type)
+		.HeaderContentPadding(FMargin(10, 1, 0, 1))
+		.SortMode(this, &SFoliagePalette::GetMeshColumnSortMode)
+		.OnSort(this, &SFoliagePalette::OnMeshesColumnSortModeChanged)
+		.DefaultLabel(this, &SFoliagePalette::GetMeshesHeaderText)
+
+		// Instance Count
+		+ SHeaderRow::Column(FoliagePaletteTreeColumns::ColumnID_InstanceCount)
+		.HeaderContentPadding(FMargin(10, 1, 0, 1))
+		.DefaultLabel(LOCTEXT("InstanceCount", "Count"))
+		.DefaultTooltip(this, &SFoliagePalette::GetTotalInstanceCountTooltipText)
+		.FixedWidth(75.f)
+
+		// Save Asset
+		+ SHeaderRow::Column(FoliagePaletteTreeColumns::ColumnID_Save)
+		.FixedWidth(24.0f)
+		.DefaultLabel(BlankText)
+	);
+
+	// View Mode Switcher
+	SAssignNew(WidgetSwitcher, SWidgetSwitcher);
+
+	// Thumbnail View
+	WidgetSwitcher->AddSlot(EFoliagePaletteViewMode::Thumbnail)
+	[
+		SNew(SScrollBorder, TileViewWidget.ToSharedRef())
+		.Content()
+		[
+			TileViewWidget.ToSharedRef()
+		]
+	];
+
+	// Tree View
+	WidgetSwitcher->AddSlot(EFoliagePaletteViewMode::Tree)
+	[
+		SNew(SScrollBorder, TreeViewWidget.ToSharedRef())
+		.Style(&FEditorStyle::Get().GetWidgetStyle<FScrollBorderStyle>("FoliageEditMode.TreeView.ScrollBorder"))
+		.Content()
+		[
+			TreeViewWidget.ToSharedRef()
+		]
+	];
+
+	WidgetSwitcher->SetActiveWidgetIndex(FoliageEditMode->UISettings.GetActivePaletteViewMode());
+
+	return WidgetSwitcher.ToSharedRef();
+}
+
+void SFoliagePalette::GetPaletteItemFilterString(FFoliagePaletteItemModelPtr PaletteItemModel, TArray<FString>& OutArray) const
+{
+	OutArray.Add(PaletteItemModel->GetFoliageTypeDisplayNameText().ToString());
+}
+
+void SFoliagePalette::OnSearchTextChanged(const FText& InFilterText)
+{
+	TypeFilter->SetRawFilterText(InFilterText);
+	UpdatePalette();
+}
+
 TSharedRef<SWidget> SFoliagePalette::GetAddFoliageTypePicker()
 {
 	TArray<const UClass*> ClassFilters;
@@ -393,16 +543,17 @@ TSharedRef<SWidget> SFoliagePalette::GetAddFoliageTypePicker()
 
 void SFoliagePalette::HandleOnToolChanged()
 {
-	RefreshMeshDetailsWidget();
+	RefreshDetailsWidget();
 }
 
-void SFoliagePalette::SetViewMode(EPaletteViewMode::Type NewViewMode)
+void SFoliagePalette::SetViewMode(EFoliagePaletteViewMode::Type NewViewMode)
 {
+	EFoliagePaletteViewMode::Type ActiveViewMode = FoliageEditMode->UISettings.GetActivePaletteViewMode();
 	if (ActiveViewMode != NewViewMode)
 	{
 		switch (NewViewMode)
 		{
-		case EPaletteViewMode::Thumbnail:
+		case EFoliagePaletteViewMode::Thumbnail:
 			// Set the tile selection to be the current tree selections
 			TileViewWidget->ClearSelection();
 			for (auto& TypeInfo : TreeViewWidget->GetSelectedItems())
@@ -411,7 +562,7 @@ void SFoliagePalette::SetViewMode(EPaletteViewMode::Type NewViewMode)
 			}
 			break;
 			
-		case EPaletteViewMode::Tree:
+		case EFoliagePaletteViewMode::Tree:
 			// Set the tree selection to be the current tile selection
 			TreeViewWidget->ClearSelection();
 			for (auto& TypeInfo : TileViewWidget->GetSelectedItems())
@@ -421,39 +572,48 @@ void SFoliagePalette::SetViewMode(EPaletteViewMode::Type NewViewMode)
 			break;
 		}
 
-		ActiveViewMode = NewViewMode;
-		WidgetSwitcher->SetActiveWidgetIndex(ActiveViewMode);
+		FoliageEditMode->UISettings.SetActivePaletteViewMode(NewViewMode);
+		WidgetSwitcher->SetActiveWidgetIndex(NewViewMode);
+		
+		RefreshActivePaletteViewWidget();
 	}
 }
 
-bool SFoliagePalette::IsActiveViewMode(EPaletteViewMode::Type ViewMode) const
+bool SFoliagePalette::IsActiveViewMode(EFoliagePaletteViewMode::Type ViewMode) const
 {
-	return ActiveViewMode == ViewMode;
+	return FoliageEditMode->UISettings.GetActivePaletteViewMode() == ViewMode;
 }
 
 void SFoliagePalette::ToggleShowTooltips()
 {
-	bShowTooltips = !bShowTooltips;
+	const bool bCurrentlyShowingTooltips = FoliageEditMode->UISettings.GetShowPaletteItemTooltips();
+	FoliageEditMode->UISettings.SetShowPaletteItemTooltips(!bCurrentlyShowingTooltips);
 }
 
 bool SFoliagePalette::ShouldShowTooltips() const
 {
-	return bShowTooltips;
+	return FoliageEditMode->UISettings.GetShowPaletteItemTooltips();
 }
 
-void SFoliagePalette::OnSelectionChanged(FFoliageMeshUIInfoPtr Item, ESelectInfo::Type SelectInfo)
+FText SFoliagePalette::GetSearchText() const
 {
-	RefreshMeshDetailsWidget();
+	return TypeFilter->GetRawFilterText();
 }
 
-void SFoliagePalette::OnItemDoubleClicked(FFoliageMeshUIInfoPtr Item) const
+void SFoliagePalette::OnSelectionChanged(FFoliagePaletteItemModelPtr Item, ESelectInfo::Type SelectInfo)
 {
-	Item->Settings->IsSelected = !Item->Settings->IsSelected;
+	RefreshDetailsWidget();
+}
+
+void SFoliagePalette::OnItemDoubleClicked(FFoliagePaletteItemModelPtr Item) const
+{
+	const bool bCurrentlyActive = Item->GetFoliageType()->IsSelected;
+	Item->SetTypeActiveInPalette(!bCurrentlyActive);
 }
 
 TSharedRef<SWidget> SFoliagePalette::GetViewOptionsMenuContent()
 {	
-	const FFoliageEditCommands& Commands = FFoliageEditCommands::Get();
+	const FFoliagePaletteCommands& Commands = FFoliagePaletteCommands::Get();
 	FMenuBuilder MenuBuilder(true, UICommandList);
 
 	MenuBuilder.BeginSection("FoliagePaletteViewMode", LOCTEXT("ViewModeHeading", "Palette View Mode"));
@@ -463,26 +623,26 @@ TSharedRef<SWidget> SFoliagePalette::GetViewOptionsMenuContent()
 			LOCTEXT("ThumbnailView_ToolTip", "Display thumbnails for each foliage type in the palette."),
 			FSlateIcon(),
 			FUIAction(
-				FExecuteAction::CreateSP(this, &SFoliagePalette::SetViewMode, EPaletteViewMode::Thumbnail),
+				FExecuteAction::CreateSP(this, &SFoliagePalette::SetViewMode, EFoliagePaletteViewMode::Thumbnail),
 				FCanExecuteAction(),
-				FIsActionChecked::CreateSP(this, &SFoliagePalette::IsActiveViewMode, EPaletteViewMode::Thumbnail)
-			),
+				FIsActionChecked::CreateSP(this, &SFoliagePalette::IsActiveViewMode, EFoliagePaletteViewMode::Thumbnail)
+				),
 			NAME_None,
 			EUserInterfaceActionType::RadioButton
-		);
+			);
 
 		MenuBuilder.AddMenuEntry(
 			LOCTEXT("ListView", "List"),
 			LOCTEXT("ListView_ToolTip", "Display foliage types in the palette as a list."),
 			FSlateIcon(),
 			FUIAction(
-				FExecuteAction::CreateSP(this, &SFoliagePalette::SetViewMode, EPaletteViewMode::Tree),
+				FExecuteAction::CreateSP(this, &SFoliagePalette::SetViewMode, EFoliagePaletteViewMode::Tree),
 				FCanExecuteAction(),
-				FIsActionChecked::CreateSP(this, &SFoliagePalette::IsActiveViewMode, EPaletteViewMode::Tree)
-			),
+				FIsActionChecked::CreateSP(this, &SFoliagePalette::IsActiveViewMode, EFoliagePaletteViewMode::Tree)
+				),
 			NAME_None,
 			EUserInterfaceActionType::RadioButton
-		);
+			);
 	}
 	MenuBuilder.EndSection();
 
@@ -496,26 +656,37 @@ TSharedRef<SWidget> SFoliagePalette::GetViewOptionsMenuContent()
 				FExecuteAction::CreateSP(this, &SFoliagePalette::ToggleShowTooltips),
 				FCanExecuteAction(),
 				FIsActionChecked::CreateSP(this, &SFoliagePalette::ShouldShowTooltips)
-			),
+				),
 			NAME_None,
 			EUserInterfaceActionType::ToggleButton
-		);
-
-		// Tile size slider
-		// Hide inactive types
+			);
+
+		MenuBuilder.AddWidget(
+			SNew(SSlider)
+				.ToolTipText( LOCTEXT("ThumbnailScaleToolTip", "Adjust the size of thumbnails.") )
+				.Value( this, &SFoliagePalette::GetThumbnailScale)
+				.OnValueChanged( this, &SFoliagePalette::SetThumbnailScale)
+				.IsEnabled( this, &SFoliagePalette::GetThumbnailScaleSliderEnabled)
+				.OnMouseCaptureEnd(this, &SFoliagePalette::RefreshActivePaletteViewWidget),
+			LOCTEXT("ThumbnailScaleLabel", "Scale"),
+			/*bNoIndent=*/true
+			);
+
+		//@todo: Hide inactive types
 	}
 	MenuBuilder.EndSection();
 
 	return MenuBuilder.MakeWidget();
 }
 
-TSharedPtr<SListView<FFoliageMeshUIInfoPtr>> SFoliagePalette::GetActiveViewWidget() const
+TSharedPtr<SListView<FFoliagePaletteItemModelPtr>> SFoliagePalette::GetActiveViewWidget() const
 {
-	if (ActiveViewMode == EPaletteViewMode::Thumbnail)
+	const EFoliagePaletteViewMode::Type ActiveViewMode = FoliageEditMode->UISettings.GetActivePaletteViewMode();
+	if (ActiveViewMode == EFoliagePaletteViewMode::Thumbnail)
 	{
 		return TileViewWidget;
 	}
-	else if (ActiveViewMode == EPaletteViewMode::Tree)
+	else if (ActiveViewMode == EFoliagePaletteViewMode::Tree)
 	{
 		return TreeViewWidget;
 	}
@@ -528,11 +699,45 @@ EVisibility SFoliagePalette::GetDropFoliageHintVisibility() const
 	return FoliageEditMode->GetFoliageMeshList().Num() == 0 ? EVisibility::Visible : EVisibility::Collapsed;
 }
 
+EVisibility SFoliagePalette::GetFoliageDropTargetVisibility() const
+{
+	if (FSlateApplication::Get().IsDragDropping())
+	{
+		TArray<FAssetData> DraggedAssets = AssetUtil::ExtractAssetDataFromDrag(FSlateApplication::Get().GetDragDroppingContent());
+		for (const FAssetData& AssetData : DraggedAssets)
+		{
+			if (AssetData.IsValid() && (AssetData.GetClass()->IsChildOf(UStaticMesh::StaticClass()) || AssetData.GetClass()->IsChildOf(UFoliageType::StaticClass())))
+			{
+				return EVisibility::Visible;
+			}
+		}
+	}
+
+	return EVisibility::Hidden;
+}
+
+FReply SFoliagePalette::HandleFoliageDropped(const FGeometry& DropZoneGeometry, const FDragDropEvent& DragDropEvent)
+{
+	TArray<FAssetData> DroppedAssetData = AssetUtil::ExtractAssetDataFromDrag(DragDropEvent);
+	if (DroppedAssetData.Num() > 0)
+	{
+		// Treat the entire drop as a transaction (in case multiples types are being added)
+		const FScopedTransaction Transaction(NSLOCTEXT("UnrealEd", "FoliageMode_DragDropTypesTransaction", "Drag-drop Foliage"));
+
+		for (auto& AssetData : DroppedAssetData)
+		{
+			AddFoliageType(AssetData);
+		}
+	}
+
+	return FReply::Handled();
+}
+
 //	CONTEXT MENU
 
 TSharedPtr<SWidget> SFoliagePalette::ConstructFoliageTypeContextMenu()
 {
-	const FFoliageEditCommands& Commands = FFoliageEditCommands::Get();
+	const FFoliagePaletteCommands& Commands = FFoliagePaletteCommands::Get();
 	FMenuBuilder MenuBuilder(true, UICommandList);
 
 	if (GetActiveViewWidget()->GetSelectedItems().Num() > 0)
@@ -582,21 +787,22 @@ TSharedPtr<SWidget> SFoliagePalette::ConstructFoliageTypeContextMenu()
 
 void SFoliagePalette::OnSaveSelectedAsFoliageType()
 {
-	for (FFoliageMeshUIInfoPtr& TypeInfo : GetActiveViewWidget()->GetSelectedItems())
+	for (FFoliagePaletteItemModelPtr& PaletteItem : GetActiveViewWidget()->GetSelectedItems())
 	{
-		UFoliageType* SavedSettings = FoliageEditMode->SaveFoliageTypeObject(TypeInfo->Settings);
-		if (SavedSettings)
+		UFoliageType* FoliageType = PaletteItem->GetFoliageType();
+		UFoliageType* SavedFoliageType = FoliageEditMode->SaveFoliageTypeObject(FoliageType);
+		if (SavedFoliageType)
 		{
-			TypeInfo->Settings = SavedSettings;
+			FoliageType = SavedFoliageType;
 		}
 	}
 }
 
 bool SFoliagePalette::OnCanSaveSelectedAsFoliageType() const
 {
-	for (FFoliageMeshUIInfoPtr& TypeInfo : GetActiveViewWidget()->GetSelectedItems())
+	for (FFoliagePaletteItemModelPtr& PaletteItem : GetActiveViewWidget()->GetSelectedItems())
 	{
-		if (TypeInfo->Settings->IsAsset())
+		if (PaletteItem->GetFoliageType()->IsAsset())
 		{
 			// At least one selected type is an asset
 			return false;
@@ -608,9 +814,9 @@ bool SFoliagePalette::OnCanSaveSelectedAsFoliageType() const
 
 bool SFoliagePalette::AreAnyNonAssetTypesSelected() const
 {
-	for (FFoliageMeshUIInfoPtr& TypeInfo : GetActiveViewWidget()->GetSelectedItems())
+	for (FFoliagePaletteItemModelPtr& PaletteItem : GetActiveViewWidget()->GetSelectedItems())
 	{
-		if (!TypeInfo->Settings->IsAsset())
+		if (!PaletteItem->GetFoliageType()->IsAsset())
 		{
 			// At least one selected type isn't an asset
 			return true;
@@ -645,15 +851,15 @@ void SFoliagePalette::OnReplaceFoliageTypeSelected(const FAssetData& AssetData)
 {
 	FSlateApplication::Get().DismissAllMenus();
 
-	UFoliageType* NewSettings = Cast<UFoliageType>(AssetData.GetAsset());
-	if (GetActiveViewWidget()->GetSelectedItems().Num() && NewSettings)
+	UFoliageType* NewFoliageType = Cast<UFoliageType>(AssetData.GetAsset());
+	if (GetActiveViewWidget()->GetSelectedItems().Num() && NewFoliageType)
 	{
-		for (FFoliageMeshUIInfoPtr& TypeInfo : GetActiveViewWidget()->GetSelectedItems())
+		for (FFoliagePaletteItemModelPtr& PaletteItem : GetActiveViewWidget()->GetSelectedItems())
 		{
-			UFoliageType* OldSettings = TypeInfo->Settings;
-			if (OldSettings != NewSettings)
+			UFoliageType* OldFoliageType = PaletteItem->GetFoliageType();
+			if (OldFoliageType != NewFoliageType)
 			{
-				FoliageEditMode->ReplaceSettingsObject(OldSettings, NewSettings);
+				FoliageEditMode->ReplaceSettingsObject(OldFoliageType, NewFoliageType);
 			}
 		}
 	}
@@ -663,10 +869,10 @@ void SFoliagePalette::OnRemoveFoliageType()
 {
 	int32 NumInstances = 0;
 	TArray<UFoliageType*> FoliageTypeList;
-	for (FFoliageMeshUIInfoPtr& TypeInfo : GetActiveViewWidget()->GetSelectedItems())
+	for (FFoliagePaletteItemModelPtr& PaletteItem : GetActiveViewWidget()->GetSelectedItems())
 	{
-		NumInstances += TypeInfo->InstanceCountTotal;
-		FoliageTypeList.Add(TypeInfo->Settings);
+		NumInstances += PaletteItem->GetTypeUIInfo()->InstanceCountTotal;
+		FoliageTypeList.Add(PaletteItem->GetFoliageType());
 	}
 
 	bool bProceed = true;
@@ -685,9 +891,9 @@ void SFoliagePalette::OnRemoveFoliageType()
 void SFoliagePalette::OnShowFoliageTypeInCB()
 {
 	TArray<UObject*> SelectedAssets;
-	for (FFoliageMeshUIInfoPtr& TypeInfo : GetActiveViewWidget()->GetSelectedItems())
+	for (FFoliagePaletteItemModelPtr& PaletteItem : GetActiveViewWidget()->GetSelectedItems())
 	{
-		UFoliageType* FoliageType = TypeInfo->Settings;
+		UFoliageType* FoliageType = PaletteItem->GetFoliageType();
 		if (FoliageType->IsAsset())
 		{
 			SelectedAssets.Add(FoliageType);
@@ -706,27 +912,27 @@ void SFoliagePalette::OnShowFoliageTypeInCB()
 
 void SFoliagePalette::OnSelectAllInstances()
 {
-	for (FFoliageMeshUIInfoPtr& TypeInfo : GetActiveViewWidget()->GetSelectedItems())
+	for (FFoliagePaletteItemModelPtr& PaletteItem : GetActiveViewWidget()->GetSelectedItems())
 	{
-		UFoliageType* FoliageType = TypeInfo->Settings;
+		UFoliageType* FoliageType = PaletteItem->GetFoliageType();
 		FoliageEditMode->SelectInstances(FoliageType, true);
 	}
 }
 
 void SFoliagePalette::OnDeselectAllInstances()
 {
-	for (FFoliageMeshUIInfoPtr& TypeInfo : GetActiveViewWidget()->GetSelectedItems())
+	for (FFoliagePaletteItemModelPtr& PaletteItem : GetActiveViewWidget()->GetSelectedItems())
 	{
-		UFoliageType* FoliageType = TypeInfo->Settings;
+		UFoliageType* FoliageType = PaletteItem->GetFoliageType();
 		FoliageEditMode->SelectInstances(FoliageType, false);
 	}
 }
 
 void SFoliagePalette::OnSelectInvalidInstances()
 {
-	for (FFoliageMeshUIInfoPtr& TypeInfo : GetActiveViewWidget()->GetSelectedItems())
+	for (FFoliagePaletteItemModelPtr& PaletteItem : GetActiveViewWidget()->GetSelectedItems())
 	{
-		const UFoliageType* FoliageType = TypeInfo->Settings;
+		UFoliageType* FoliageType = PaletteItem->GetFoliageType();
 		FoliageEditMode->SelectInvalidInstances(FoliageType);
 	}
 }
@@ -738,21 +944,40 @@ bool SFoliagePalette::CanSelectInstances() const
 
 // THUMBNAIL VIEW
 
-TSharedRef<ITableRow> SFoliagePalette::GenerateTile(FFoliageMeshUIInfoPtr Item, const TSharedRef<STableViewBase>& OwnerTable)
+TSharedRef<ITableRow> SFoliagePalette::GenerateTile(FFoliagePaletteItemModelPtr Item, const TSharedRef<STableViewBase>& OwnerTable)
 {
-	TSharedPtr<FFoliagePaletteItemModel> ItemModel = MakeShareable(new FFoliagePaletteItemModel(Item, SharedThis(this), FoliageEditMode));
-	return SNew(SFoliagePaletteItemTile, OwnerTable, ItemModel);
+	return SNew(SFoliagePaletteItemTile, OwnerTable, Item);
+}
+
+float SFoliagePalette::GetScaledThumbnailSize() const
+{
+	const FInt32Interval& SizeRange = FoliagePaletteConstants::ThumbnailSizeRange;
+	return SizeRange.Min + SizeRange.Size() * FoliageEditMode->UISettings.GetPaletteThumbnailScale();
+}
+
+float SFoliagePalette::GetThumbnailScale() const
+{
+	return FoliageEditMode->UISettings.GetPaletteThumbnailScale();
+}
+
+void SFoliagePalette::SetThumbnailScale(float InScale)
+{
+	FoliageEditMode->UISettings.SetPaletteThumbnailScale(InScale);
+}
+
+bool SFoliagePalette::GetThumbnailScaleSliderEnabled() const
+{
+	return FoliageEditMode->UISettings.GetActivePaletteViewMode() == EFoliagePaletteViewMode::Thumbnail;
 }
 
 // TREE VIEW
 
-TSharedRef<ITableRow> SFoliagePalette::TreeViewGenerateRow(FFoliageMeshUIInfoPtr Item, const TSharedRef<STableViewBase>& OwnerTable)
+TSharedRef<ITableRow> SFoliagePalette::TreeViewGenerateRow(FFoliagePaletteItemModelPtr Item, const TSharedRef<STableViewBase>& OwnerTable)
 {
-	TSharedPtr<FFoliagePaletteItemModel> ItemModel = MakeShareable(new FFoliagePaletteItemModel(Item, SharedThis(this), FoliageEditMode));
-	return SNew(SFoliagePaletteItemRow, OwnerTable, ItemModel);
+	return SNew(SFoliagePaletteItemRow, OwnerTable, Item);
 }
 
-void SFoliagePalette::TreeViewGetChildren(FFoliageMeshUIInfoPtr Item, TArray<FFoliageMeshUIInfoPtr>& OutChildren)
+void SFoliagePalette::TreeViewGetChildren(FFoliagePaletteItemModelPtr Item, TArray<FFoliagePaletteItemModelPtr>& OutChildren)
 {
 	//OutChildren = Item->GetChildren();
 }
@@ -762,9 +987,10 @@ ECheckBoxState SFoliagePalette::GetState_AllMeshes() const
 	bool bHasChecked = false;
 	bool bHasUnchecked = false;
 
-	for (const FFoliageMeshUIInfoPtr& TypeInfo : FoliageEditMode->GetFoliageMeshList())
+	for (FFoliagePaletteItemModelPtr& PaletteItem : GetActiveViewWidget()->GetSelectedItems())
 	{
-		if (TypeInfo->Settings->IsSelected)
+		UFoliageType* FoliageType = PaletteItem->GetFoliageType();
+		if (FoliageType->IsSelected)
 		{
 			bHasChecked = true;
 		}
@@ -784,9 +1010,10 @@ ECheckBoxState SFoliagePalette::GetState_AllMeshes() const
 
 void SFoliagePalette::OnCheckStateChanged_AllMeshes(ECheckBoxState InState)
 {
-	for (const FFoliageMeshUIInfoPtr& TypeInfo : FoliageEditMode->GetFoliageMeshList())
+	for (FFoliagePaletteItemModelPtr& PaletteItem : GetActiveViewWidget()->GetSelectedItems())
 	{
-		TypeInfo->Settings->IsSelected = (InState == ECheckBoxState::Checked);
+		UFoliageType* FoliageType = PaletteItem->GetFoliageType();
+		FoliageType->IsSelected = (InState == ECheckBoxState::Checked);
 	}
 }
 
@@ -819,13 +1046,13 @@ FText SFoliagePalette::GetTotalInstanceCountTooltipText() const
 
 //	DETAILS VIEW
 
-void SFoliagePalette::RefreshMeshDetailsWidget()
+void SFoliagePalette::RefreshDetailsWidget()
 {
 	TArray<UObject*> SelectedFoliageTypes;
 
-	for (const auto& Info : GetActiveViewWidget()->GetSelectedItems())
+	for (const auto& PaletteItem : GetActiveViewWidget()->GetSelectedItems())
 	{
-		SelectedFoliageTypes.Add(Info->Settings);
+		SelectedFoliageTypes.Add(PaletteItem->GetFoliageType());
 	}
 
 	const bool bForceRefresh = true;
@@ -836,11 +1063,11 @@ FText SFoliagePalette::GetDetailsNameAreaText() const
 {
 	FText OutText;
 
-	auto SelectedTypes = GetActiveViewWidget()->GetSelectedItems();
-	if (SelectedTypes.Num() == 1)
+	auto SelectedItems = GetActiveViewWidget()->GetSelectedItems();
+	if (SelectedItems.Num() == 1)
 	{
 		FName DisplayName;
-		UFoliageType* SelectedType = SelectedTypes[0]->Settings;
+		UFoliageType* SelectedType = SelectedItems[0]->GetFoliageType();
 		if (SelectedType->IsAsset())
 		{
 			DisplayName = SelectedType->GetFName();
@@ -851,9 +1078,9 @@ FText SFoliagePalette::GetDetailsNameAreaText() const
 		}
 		OutText = FText::FromName(DisplayName);
 	}
-	else if (SelectedTypes.Num() > 1)
+	else if (SelectedItems.Num() > 1)
 	{
-		OutText = FText::Format(LOCTEXT("DetailsNameAreaText_Multiple", "{0} Types Selected"), FText::AsNumber(SelectedTypes.Num()));
+		OutText = FText::Format(LOCTEXT("DetailsNameAreaText_Multiple", "{0} Types Selected"), FText::AsNumber(SelectedItems.Num()));
 	}
 
 	return OutText;
@@ -875,8 +1102,59 @@ FReply SFoliagePalette::OnShowHideDetailsClicked() const
 {
 	const bool bDetailsCurrentlyVisible = DetailsWidget->GetVisibility() != EVisibility::Collapsed;
 	DetailsWidget->SetVisibility(bDetailsCurrentlyVisible ? EVisibility::Collapsed : EVisibility::SelfHitTestInvisible);
+	FoliageEditMode->UISettings.SetShowPaletteItemDetails(!bDetailsCurrentlyVisible);
 
 	return FReply::Handled();
 }
 
+EActiveTimerReturnType SFoliagePalette::UpdatePaletteItems(double InCurrentTime, float InDeltaTime)
+{
+	if (bItemsNeedRefresh)
+	{
+		const auto& IFATypeList = FoliageEditMode->GetFoliageMeshList();
+
+		// Check if any palette items should be removed
+		// We do this instead of a full rebuild to avoid the flicker that occurs when creating a new thumbnail widget
+		TArray<FFoliageMeshUIInfoPtr> ItemsToRemove;
+		for (auto& ItemPair : PaletteItemsByTypeInfo)
+		{
+			if (!IFATypeList.Contains(ItemPair.Key))
+			{
+				ItemsToRemove.Add(ItemPair.Key);
+			}
+		}
+
+		// Remove the items from the palette
+		for (FFoliageMeshUIInfoPtr& Item : ItemsToRemove)
+		{
+			PaletteItemsByTypeInfo.Remove(Item);
+		}
+
+		// Create a palette item for any newly added types
+		for (FFoliageMeshUIInfoPtr& TypeInfo : FoliageEditMode->GetFoliageMeshList())
+		{
+			if (!PaletteItemsByTypeInfo.Contains(TypeInfo))
+			{
+				PaletteItemsByTypeInfo.Add(TypeInfo) = MakeShareable(new FFoliagePaletteItemModel(TypeInfo, SharedThis(this), FoliageEditMode));
+			}
+		}
+	}
+
+	// Update the filtered items
+	FilteredItems.Empty();
+	for (auto& ItemPair : PaletteItemsByTypeInfo)
+	{
+		if (TypeFilter->PassesFilter(ItemPair.Value))
+		{
+			FilteredItems.Add(ItemPair.Value);
+		}
+	}
+
+	// Refresh the appropriate view
+	RefreshActivePaletteViewWidget();
+
+	bIsActiveTimerRegistered = false;
+	return EActiveTimerReturnType::Stop;
+}
+
 #undef LOCTEXT_NAMESPACE
\ No newline at end of file
diff --git a/Engine/Source/Editor/FoliageEdit/Private/SFoliagePalette.h b/Engine/Source/Editor/FoliageEdit/Private/SFoliagePalette.h
--- a/Engine/Source/Editor/FoliageEdit/Private/SFoliagePalette.h
+++ b/Engine/Source/Editor/FoliageEdit/Private/SFoliagePalette.h
@@ -2,22 +2,19 @@
 
 #pragma once
 
+#include "Misc/TextFilter.h"
+
 class SFoliagePaletteTile;
 class FEdModeFoliage;
-struct FFoliageMeshUIInfo;
-typedef TSharedPtr<FFoliageMeshUIInfo> FFoliageMeshUIInfoPtr; //should match typedef in FoliageEdMode.h
+class FFoliagePaletteItemModel;
 
-typedef STreeView<FFoliageMeshUIInfoPtr> SFoliageTypeTreeView;
-typedef STileView<FFoliageMeshUIInfoPtr> SFoliageTypeTileView;
+typedef TSharedPtr<FFoliagePaletteItemModel> FFoliagePaletteItemModelPtr;
+typedef STreeView<FFoliagePaletteItemModelPtr> SFoliageTypeTreeView;
+typedef STileView<FFoliagePaletteItemModelPtr> SFoliageTypeTileView;
 
-/** View modes supported by the foliage palette */
-namespace EPaletteViewMode
+namespace FoliagePaletteConstants
 {
-	enum Type
-	{
-		Thumbnail,
-		Tree
-	};
+	const FInt32Interval ThumbnailSizeRange(32, 128);
 }
 
 /** The palette of foliage types available for use by the foliage edit mode */
@@ -29,25 +26,43 @@ public:
 	SLATE_END_ARGS()
 
 	void Construct(const FArguments& InArgs);
-	void Refresh();
+	virtual FReply OnKeyDown(const FGeometry& MyGeometry, const FKeyEvent& InKeyEvent) override;
+
+	/** Updates the foliage palette, optionally doing a full refresh of the items in the palette as well */
+	void UpdatePalette(bool bRefreshItems = false);
 	
 	bool AnySelectedTileHovered() const;
 	void ActivateAllSelectedTypes(bool bActivate) const;
 
 	/** @return True if the given view mode is the active view mode */
-	bool IsActiveViewMode(EPaletteViewMode::Type ViewMode) const;
+	bool IsActiveViewMode(EFoliagePaletteViewMode::Type ViewMode) const;
 
 	/** @return True if tooltips should be shown when hovering over foliage type items in the palette */
 	bool ShouldShowTooltips() const;
 
+	/** @return The current search filter text */
+	FText GetSearchText() const;
+
+	/** Adds the foliage type asset to the instanced foliage actor's list of types. */
+	void AddFoliageType(const FAssetData& AssetData);
+
 private:	// GENERAL
 
 	/** Binds commands used by the palette */
 	void BindCommands();
 
-	/** Adds the foliage type asset to the foliage actor's list of meshes. */
-	void AddFoliageType(const FAssetData& AssetData);
-	
+	/** Refreshes the active palette view widget */
+	void RefreshActivePaletteViewWidget();
+
+	/** Creates the palette views */
+	TSharedRef<class SWidgetSwitcher> CreatePaletteViews();
+
+	/** Adds the displayed name of the foliage type for filtering */
+	void GetPaletteItemFilterString(FFoliagePaletteItemModelPtr PaletteItemModel, TArray<FString>& OutArray) const;
+
+	/** Handles changes to the search filter text */
+	void OnSearchTextChanged(const FText& InFilterText);
+
 	/** Gets the asset picker for adding a foliage type. */
 	TSharedRef<SWidget> GetAddFoliageTypePicker();
 
@@ -62,7 +77,7 @@ private:	// GENERAL
 	void HandleOnToolChanged();
 
 	/** Sets the view mode of the palette */
-	void SetViewMode(EPaletteViewMode::Type NewViewMode);
+	void SetViewMode(EFoliagePaletteViewMode::Type NewViewMode);
 
 	/** Sets whether to show tooltips when hovering over foliage type items in the palette */
 	void ToggleShowTooltips();
@@ -74,19 +89,25 @@ private:	// GENERAL
 	int32 GetActiveViewIndex() const;
 
 	/** Handler for selection changes in either view */
-	void OnSelectionChanged(FFoliageMeshUIInfoPtr Item, ESelectInfo::Type SelectInfo);
+	void OnSelectionChanged(FFoliagePaletteItemModelPtr Item, ESelectInfo::Type SelectInfo);
 
 	/** Toggle the activation state of a type on a double-click */
-	void OnItemDoubleClicked(FFoliageMeshUIInfoPtr Item) const;
+	void OnItemDoubleClicked(FFoliagePaletteItemModelPtr Item) const;
 
 	/** Creates the view options menu */
 	TSharedRef<SWidget> GetViewOptionsMenuContent();
 
-	TSharedPtr<SListView<FFoliageMeshUIInfoPtr>> GetActiveViewWidget() const;
+	TSharedPtr<SListView<FFoliagePaletteItemModelPtr>> GetActiveViewWidget() const;
 
 	/** Gets the visibility of the "Drop Foliage Here" prompt for when the palette is empty */
 	EVisibility GetDropFoliageHintVisibility() const;
 
+	/** Gets the visibility of the drag-drop zone overlay */
+	EVisibility GetFoliageDropTargetVisibility() const;
+
+	/** Handles dropping of a mesh or foliage type into the palette */
+	FReply HandleFoliageDropped(const FGeometry& DropZoneGeometry, const FDragDropEvent& DragDropEvent);
+
 private:	// CONTEXT MENU
 
 	/** @return the SWidget containing the context menu */
@@ -128,15 +149,27 @@ private:	// CONTEXT MENU
 private:	// THUMBNAIL VIEW
 
 	/** Creates a thumbnail tile for the given foliage type */
-	TSharedRef<ITableRow> GenerateTile(FFoliageMeshUIInfoPtr Item, const TSharedRef<STableViewBase>& OwnerTable);
+	TSharedRef<ITableRow> GenerateTile(FFoliagePaletteItemModelPtr Item, const TSharedRef<STableViewBase>& OwnerTable);
+
+	/** Gets the scaled thumbnail tile size */
+	float GetScaledThumbnailSize() const;
+
+	/** Gets the current scale of the thumbnail tiles */
+	float GetThumbnailScale() const;
+
+	/** Sets the current scale of the thumbnail tiles */
+	void SetThumbnailScale(float InScale);
+
+	/** Gets whether the thumbnail scaling slider is visible */
+	bool GetThumbnailScaleSliderEnabled() const;
 
 private:	// TREE VIEW
 
 	/** Generates a row widget for foliage mesh item */
-	TSharedRef<ITableRow> TreeViewGenerateRow(FFoliageMeshUIInfoPtr Item, const TSharedRef<STableViewBase>& OwnerTable);
+	TSharedRef<ITableRow> TreeViewGenerateRow(FFoliagePaletteItemModelPtr Item, const TSharedRef<STableViewBase>& OwnerTable);
 
 	/** Generates a list of children items for foliage item */
-	void TreeViewGetChildren(FFoliageMeshUIInfoPtr Item, TArray<FFoliageMeshUIInfoPtr>& OutChildren);
+	void TreeViewGetChildren(FFoliagePaletteItemModelPtr Item, TArray<FFoliagePaletteItemModelPtr>& OutChildren);
 
 	/** Text for foliage meshes list header */
 	FText GetMeshesHeaderText() const;
@@ -151,7 +184,7 @@ private:	// TREE VIEW
 private:	// DETAILS
 
 	/** Refreshes the mesh details widget to match the current selection */
-	void RefreshMeshDetailsWidget();
+	void RefreshDetailsWidget();
 
 	/** Gets the text for the details area header */
 	FText GetDetailsNameAreaText() const;
@@ -166,6 +199,18 @@ private:	// DETAILS
 	FReply OnShowHideDetailsClicked() const;
 
 private:
+	/** Active timer handler to update the items in the palette */
+	EActiveTimerReturnType UpdatePaletteItems(double InCurrentTime, float InDeltaTime);
+
+private:
+	typedef TTextFilter<FFoliagePaletteItemModelPtr> FoliageTypeTextFilter;
+	TSharedPtr<FoliageTypeTextFilter> TypeFilter;
+
+	/** All the items in the palette */
+	TMap<FFoliageMeshUIInfoPtr, FFoliagePaletteItemModelPtr> PaletteItemsByTypeInfo;
+
+	/** The filtered list of types to display in the palette */
+	TArray<FFoliagePaletteItemModelPtr> FilteredItems;
 
 	/** Switches between the thumbnail and tree views */
 	TSharedPtr<class SWidgetSwitcher> WidgetSwitcher;
@@ -190,7 +235,6 @@ private:
 
 	FEdModeFoliage* FoliageEditMode;
 
-	//@todo: Should be saved to config
-	EPaletteViewMode::Type ActiveViewMode;
-	bool bShowTooltips;
+	bool bItemsNeedRefresh : 1;
+	bool bIsActiveTimerRegistered : 1;
 };
\ No newline at end of file
