[UE-14596] - Fixed crash trying to get widget for offscreen list items in foliage palette

diff --git a/Engine/Source/Editor/FoliageEdit/Private/SFoliagePalette.cpp b/Engine/Source/Editor/FoliageEdit/Private/SFoliagePalette.cpp
--- a/Engine/Source/Editor/FoliageEdit/Private/SFoliagePalette.cpp
+++ b/Engine/Source/Editor/FoliageEdit/Private/SFoliagePalette.cpp
@@ -242,7 +242,6 @@ void SFoliagePalette::Construct(const FArguments& InArgs)
 
 								// Arrow
 								+ SHorizontalBox::Slot()
-								.Padding(FMargin(3.f, 0.f))
 								.AutoWidth()
 								.HAlign(HAlign_Center)
 								.VAlign(VAlign_Center)
@@ -254,6 +253,7 @@ void SFoliagePalette::Construct(const FArguments& InArgs)
 								// Details icon
 								+SHorizontalBox::Slot()
 								.AutoWidth()
+								.Padding(FMargin(3.f, 0.f))
 								.HAlign(HAlign_Center)
 								.VAlign(VAlign_Center)
 								[
@@ -348,7 +348,8 @@ bool SFoliagePalette::AnySelectedTileHovered() const
 {
 	for (auto& TypeInfo : GetActiveViewWidget()->GetSelectedItems())
 	{
-		if (TileViewWidget->WidgetFromItem(TypeInfo)->AsWidget()->IsHovered())
+		TSharedPtr<ITableRow> Tile = TileViewWidget->WidgetFromItem(TypeInfo);
+		if (Tile.IsValid() && Tile->AsWidget()->IsHovered())
 		{
 			return true;
 		}
