[UE-14644] - Changing the mesh assigned to a foliage type will propagate to instances in the world & update the entry in the palette
- Fixed crash on shutdown when foliage mode was active

Various other important & related usability fixes (for 4.8):
- Palette selection is retained accross undo/redo
- To maintain item order at all times, the default sort mode for the foliage palette isn't none
- Restored previous layout of view options and details toggle widgets
- Added commands for activate & deactivate selected; removed the activation checkbox when tiles are tiny

diff --git a/Engine/Source/Editor/FoliageEdit/Private/FoliageEdMode.cpp b/Engine/Source/Editor/FoliageEdit/Private/FoliageEdMode.cpp
--- a/Engine/Source/Editor/FoliageEdit/Private/FoliageEdMode.cpp
+++ b/Engine/Source/Editor/FoliageEdit/Private/FoliageEdMode.cpp
@@ -156,7 +156,7 @@ FEdModeFoliage::FEdModeFoliage()
 	, bToolActive(false)
 	, bCanAltDrag(false)
 	, bAdjustBrushRadius(false)
-	, FoliageMeshListSortMode(EColumnSortMode::None)
+	, FoliageMeshListSortMode(EColumnSortMode::Ascending)
 {
 	// Load resources and construct brush component
 	UMaterial* BrushMaterial = nullptr;
@@ -298,6 +298,23 @@ void FEdModeFoliage::Enter()
 	{
 		ApplySelectionToComponents(GetWorld(), true);
 	}
+
+	// Subscribe to mesh changed events (for existing and future IFA's)
+	UWorld* World = GetWorld();
+	OnActorSpawnedHandle = World->AddOnActorSpawnedHandler(FOnActorSpawned::FDelegate::CreateRaw(this, &FEdModeFoliage::HandleOnActorSpawned));
+	const int32 NumLevels = World->GetNumLevels();
+	for (int32 LevelIdx = 0; LevelIdx < NumLevels; ++LevelIdx)
+	{
+		ULevel* Level = World->GetLevel(LevelIdx);
+		if (Level && Level->bIsVisible)
+		{
+			AInstancedFoliageActor* IFA = AInstancedFoliageActor::GetInstancedFoliageActorForLevel(Level);
+			if (IFA)
+			{
+				IFA->OnFoliageTypeMeshChanged().AddSP(this, &FEdModeFoliage::HandleOnFoliageTypeMeshChanged);
+			}
+		}
+	}
 		
 	// Update UI
 	NotifyNewCurrentLevel();
@@ -333,6 +350,23 @@ void FEdModeFoliage::Exit()
 	// Clear selection visualization on any foliage items
 	ApplySelectionToComponents(GetWorld(), false);
 
+	// Remove all event subscriptions
+	UWorld* World = GetWorld();
+	World->RemoveOnActorSpawnedHandler(OnActorSpawnedHandle);
+	const int32 NumLevels = World->GetNumLevels();
+	for (int32 LevelIdx = 0; LevelIdx < NumLevels; ++LevelIdx)
+	{
+		ULevel* Level = World->GetLevel(LevelIdx);
+		if (Level && Level->bIsVisible)
+		{
+			AInstancedFoliageActor* IFA = AInstancedFoliageActor::GetInstancedFoliageActorForLevel(Level);
+			if (IFA)
+			{
+				IFA->OnFoliageTypeMeshChanged().RemoveAll(this);
+			}
+		}
+	}
+
 	// Call base Exit method to ensure proper cleanup
 	FEdMode::Exit();
 }
@@ -2218,6 +2252,20 @@ void FEdModeFoliage::SnapSelectedInstancesToGround(UWorld* InWorld)
 	}
 }
 
+void FEdModeFoliage::HandleOnActorSpawned(AActor* Actor)
+{
+	if (AInstancedFoliageActor* IFA = Cast<AInstancedFoliageActor>(Actor))
+	{
+		// If an IFA was created, we want to be notified if any meshes assigned to its foliage types change
+		IFA->OnFoliageTypeMeshChanged().AddSP(this, &FEdModeFoliage::HandleOnFoliageTypeMeshChanged);
+	}
+}
+
+void FEdModeFoliage::HandleOnFoliageTypeMeshChanged(UFoliageType* FoliageType)
+{
+	StaticCastSharedPtr<FFoliageEdModeToolkit>(Toolkit)->NotifyFoliageTypeMeshChanged(FoliageType);
+}
+
 bool FEdModeFoliage::SnapInstanceToGround(AInstancedFoliageActor* InIFA, float AlignMaxAngle, FFoliageMeshInfo& Mesh, int32 InstanceIdx)
 {
 	FFoliageInstance& Instance = Mesh.Instances[InstanceIdx];
diff --git a/Engine/Source/Editor/FoliageEdit/Private/FoliageEdMode.h b/Engine/Source/Editor/FoliageEdit/Private/FoliageEdMode.h
--- a/Engine/Source/Editor/FoliageEdit/Private/FoliageEdMode.h
+++ b/Engine/Source/Editor/FoliageEdit/Private/FoliageEdMode.h
@@ -8,6 +8,7 @@ class ULandscapeComponent;
 class UFoliageType;
 
 #include "InstancedFoliage.h"
+#include "InstancedFoliageActor.h"
 
 //
 // Forward declarations.
@@ -441,6 +442,12 @@ private:
 	bool SnapInstanceToGround(AInstancedFoliageActor* InIFA, float AlignMaxAngle, FFoliageMeshInfo& Mesh, int32 InstanceIdx);
 	void SnapSelectedInstancesToGround(UWorld* InWorld);
 
+	/** Callback for when an actor is spawned (to check if it's a new IFA) */
+	void HandleOnActorSpawned(AActor* Actor);
+
+	/** Callback for when the mesh assigned to a foliage type referenced by an IFA is changed */
+	void HandleOnFoliageTypeMeshChanged(UFoliageType* FoliageType);
+
 	/** Common code for adding instances to world based on settings */
 	static void AddInstancesImp(UWorld* InWorld, const UFoliageType* Settings, const TArray<FDesiredFoliageInstance>& DesiredInstances, const TArray<int32>& ExistingInstances = TArray<int32>(), const float Pressure = 1.f, LandscapeLayerCacheData* LandscapeLayerCaches = nullptr, const FFoliageUISettings* UISettings = nullptr);
 
@@ -473,5 +480,7 @@ private:
 
 	TArray<FFoliageMeshUIInfoPtr>	FoliageMeshList;
 	EColumnSortMode::Type			FoliageMeshListSortMode; 
+
+	FDelegateHandle OnActorSpawnedHandle;
 };
 
diff --git a/Engine/Source/Editor/FoliageEdit/Private/FoliageEdModeToolkit.cpp b/Engine/Source/Editor/FoliageEdit/Private/FoliageEdModeToolkit.cpp
--- a/Engine/Source/Editor/FoliageEdit/Private/FoliageEdModeToolkit.cpp
+++ b/Engine/Source/Editor/FoliageEdit/Private/FoliageEdModeToolkit.cpp
@@ -48,4 +48,9 @@ void FFoliageEdModeToolkit::RefreshFullList()
 	FoliageEdWidget->RefreshFullList();
 }
 
+void FFoliageEdModeToolkit::NotifyFoliageTypeMeshChanged(class UFoliageType* FoliageType)
+{
+	FoliageEdWidget->NotifyFoliageTypeMeshChanged(FoliageType);
+}
+
 #undef LOCTEXT_NAMESPACE
diff --git a/Engine/Source/Editor/FoliageEdit/Private/FoliageEdModeToolkit.h b/Engine/Source/Editor/FoliageEdit/Private/FoliageEdModeToolkit.h
--- a/Engine/Source/Editor/FoliageEdit/Private/FoliageEdModeToolkit.h
+++ b/Engine/Source/Editor/FoliageEdit/Private/FoliageEdModeToolkit.h
@@ -23,6 +23,7 @@ public:
 	virtual TSharedPtr<class SWidget> GetInlineContent() const override;
 
 	void RefreshFullList();
+	void NotifyFoliageTypeMeshChanged(class UFoliageType* FoliageType);
 
 private:
 	TSharedPtr< class SFoliageEdit > FoliageEdWidget;
diff --git a/Engine/Source/Editor/FoliageEdit/Private/FoliagePaletteCommands.cpp b/Engine/Source/Editor/FoliageEdit/Private/FoliagePaletteCommands.cpp
--- a/Engine/Source/Editor/FoliageEdit/Private/FoliagePaletteCommands.cpp
+++ b/Engine/Source/Editor/FoliageEdit/Private/FoliagePaletteCommands.cpp
@@ -7,6 +7,8 @@
 
 void FFoliagePaletteCommands::RegisterCommands()
 {
+	UI_COMMAND(ActivateFoliageType, "Activate Selected", "Sets the selected foliage types in the palette as active (i.e. included in brush actions).", EUserInterfaceActionType::Button, FInputChord());
+	UI_COMMAND(DeactivateFoliageType, "Deactivate Selected", "Sets the selected foliage types in the palette as inactive (i.e. excluded in brush actions).", EUserInterfaceActionType::Button, FInputChord());
 	UI_COMMAND(RemoveFoliageType, "Remove", "Remove this foliage type from the palette. Removes all associated instances as well.", EUserInterfaceActionType::Button, FInputChord(EKeys::Delete));
 	UI_COMMAND(ShowFoliageTypeInCB, "Show in Content Browser", "Show asset in Content Browser.", EUserInterfaceActionType::Button, FInputChord());
 	UI_COMMAND(SelectAllInstances, "Select All Instances", "Select all instances of this foliage type (must be in a selection mode).", EUserInterfaceActionType::Button, FInputChord());
diff --git a/Engine/Source/Editor/FoliageEdit/Private/FoliagePaletteCommands.h b/Engine/Source/Editor/FoliageEdit/Private/FoliagePaletteCommands.h
--- a/Engine/Source/Editor/FoliageEdit/Private/FoliagePaletteCommands.h
+++ b/Engine/Source/Editor/FoliageEdit/Private/FoliagePaletteCommands.h
@@ -23,6 +23,8 @@ public:
 	virtual void RegisterCommands() override;
 
 	/** Foliage Palette Commands */
+	TSharedPtr< FUICommandInfo > ActivateFoliageType;
+	TSharedPtr< FUICommandInfo > DeactivateFoliageType;
 	TSharedPtr< FUICommandInfo > RemoveFoliageType;
 	TSharedPtr< FUICommandInfo > ShowFoliageTypeInCB;
 	TSharedPtr< FUICommandInfo > SelectAllInstances;
diff --git a/Engine/Source/Editor/FoliageEdit/Private/FoliagePaletteItem.cpp b/Engine/Source/Editor/FoliageEdit/Private/FoliagePaletteItem.cpp
--- a/Engine/Source/Editor/FoliageEdit/Private/FoliagePaletteItem.cpp
+++ b/Engine/Source/Editor/FoliageEdit/Private/FoliagePaletteItem.cpp
@@ -45,9 +45,9 @@ FFoliagePaletteItemModel::FFoliagePaletteItemModel(FFoliageMeshUIInfoPtr InTypeI
 	ThumbnailWidget = Thumbnail->MakeThumbnailWidget(ThumbnailConfig);
 }
 
-const TSharedPtr<SFoliagePalette>& FFoliagePaletteItemModel::GetFoliagePalette() const
+TSharedPtr<SFoliagePalette> FFoliagePaletteItemModel::GetFoliagePalette() const
 {
-	return FoliagePalette;
+	return FoliagePalette.Pin();
 }
 
 FFoliageMeshUIInfoPtr FFoliagePaletteItemModel::GetTypeUIInfo() const
@@ -203,7 +203,14 @@ FText FFoliagePaletteItemModel::GetFoliageTypeDisplayNameText() const
 
 FText FFoliagePaletteItemModel::GetPaletteSearchText() const
 {
-	return FoliagePalette->GetSearchText();
+	if (FoliagePalette.IsValid())
+	{
+		return FoliagePalette.Pin()->GetSearchText();
+	}
+	else
+	{
+		return FText();
+	}
 }
 
 FText FFoliagePaletteItemModel::GetInstanceCountText(bool bRounded) const
@@ -257,6 +264,11 @@ void FFoliagePaletteItemModel::SetTypeActiveInPalette(bool bSetActiveInPalette)
 	}
 }
 
+bool FFoliagePaletteItemModel::IsActive() const
+{
+	return TypeInfo->Settings->IsSelected;
+}
+
 void FFoliagePaletteItemModel::HandleCheckStateChanged(const ECheckBoxState NewCheckedState, TAttribute<bool> IsItemWidgetSelected)
 {
 	if (!IsItemWidgetSelected.IsSet()) { return; }
@@ -266,9 +278,9 @@ void FFoliagePaletteItemModel::HandleCheckStateChanged(const ECheckBoxState NewC
 	{
 		SetTypeActiveInPalette(bShouldActivate);
 	}
-	else
+	else if (FoliagePalette.IsValid())
 	{
-		FoliagePalette->ActivateAllSelectedTypes(bShouldActivate);
+		FoliagePalette.Pin()->ActivateAllSelectedTypes(bShouldActivate);
 	}
 }
 
@@ -298,12 +310,12 @@ FReply FFoliagePaletteItemModel::HandleSaveAsset()
 
 EVisibility FFoliagePaletteItemModel::GetTooltipVisibility() const
 {
-	return FoliagePalette->ShouldShowTooltips() ? EVisibility::SelfHitTestInvisible : EVisibility::Collapsed;
+	return (FoliagePalette.IsValid() && FoliagePalette.Pin()->ShouldShowTooltips()) ? EVisibility::SelfHitTestInvisible : EVisibility::Collapsed;
 }
 
 EVisibility FFoliagePaletteItemModel::GetTooltipThumbnailVisibility() const
 {
-	return FoliagePalette->IsActiveViewMode(EFoliagePaletteViewMode::Tree) ? EVisibility::SelfHitTestInvisible : EVisibility::Collapsed;
+	return (FoliagePalette.IsValid() && FoliagePalette.Pin()->IsActiveViewMode(EFoliagePaletteViewMode::Tree)) ? EVisibility::SelfHitTestInvisible : EVisibility::Collapsed;
 }
 
 ////////////////////////////////////////////////
@@ -390,7 +402,7 @@ FLinearColor SFoliagePaletteItemTile::GetTileColorAndOpacity() const
 
 EVisibility SFoliagePaletteItemTile::GetCheckBoxVisibility() const
 {
-	return IsHovered() || (IsSelected() && Model->GetFoliagePalette()->AnySelectedTileHovered()) ? EVisibility::Visible : EVisibility::Collapsed;
+	return CanShowOverlayItems() && (IsHovered() || (IsSelected() && Model->GetFoliagePalette()->AnySelectedTileHovered())) ? EVisibility::Visible : EVisibility::Collapsed;
 }
 
 EVisibility SFoliagePaletteItemTile::GetSaveButtonVisibility() const
diff --git a/Engine/Source/Editor/FoliageEdit/Private/FoliagePaletteItem.h b/Engine/Source/Editor/FoliageEdit/Private/FoliagePaletteItem.h
--- a/Engine/Source/Editor/FoliageEdit/Private/FoliagePaletteItem.h
+++ b/Engine/Source/Editor/FoliageEdit/Private/FoliagePaletteItem.h
@@ -26,7 +26,7 @@ public:
 	FFoliagePaletteItemModel(FFoliageMeshUIInfoPtr InTypeInfo, TSharedRef<SFoliagePalette> InFoliagePalette, FEdModeFoliage* InFoliageEditMode);
 
 	/** @return The foliage palette that contains the item */
-	const TSharedPtr<SFoliagePalette>& GetFoliagePalette() const;
+	TSharedPtr<SFoliagePalette> GetFoliagePalette() const;
 
 	FFoliageMeshUIInfoPtr GetTypeUIInfo() const;
 	UFoliageType* GetFoliageType() const;
@@ -56,6 +56,9 @@ public:
 	/** Sets whether this foliage type is active in the palette */
 	void SetTypeActiveInPalette(bool bSetActiveInPalette);
 
+	/** Gets whether this foliage type is active in the palette */
+	bool IsActive() const;
+
 private:
 	/** Handles the change in activation of the item in the palette */
 	void HandleCheckStateChanged(const ECheckBoxState NewCheckedState, TAttribute<bool> IsItemWidgetSelected);
@@ -79,7 +82,7 @@ private:
 	TSharedPtr<SWidget> ThumbnailWidget;
 
 	FFoliageMeshUIInfoPtr TypeInfo;
-	TSharedPtr<SFoliagePalette> FoliagePalette;
+	TWeakPtr<SFoliagePalette> FoliagePalette;
 	FEdModeFoliage* FoliageEditMode;
 };
 
diff --git a/Engine/Source/Editor/FoliageEdit/Private/SFoliageEdit.cpp b/Engine/Source/Editor/FoliageEdit/Private/SFoliageEdit.cpp
--- a/Engine/Source/Editor/FoliageEdit/Private/SFoliageEdit.cpp
+++ b/Engine/Source/Editor/FoliageEdit/Private/SFoliageEdit.cpp
@@ -353,6 +353,11 @@ void SFoliageEdit::RefreshFullList()
 	FoliagePalette->UpdatePalette(true);
 }
 
+void SFoliageEdit::NotifyFoliageTypeMeshChanged(UFoliageType* FoliageType)
+{
+	FoliagePalette->UpdateThumbnailForType(FoliageType);
+}
+
 TSharedRef<SWidget> SFoliageEdit::BuildToolBar()
 {
 	FToolBarBuilder Toolbar(FoliageEditMode->UICommandList, FMultiBoxCustomization::None, nullptr, Orient_Vertical);
diff --git a/Engine/Source/Editor/FoliageEdit/Private/SFoliageEdit.h b/Engine/Source/Editor/FoliageEdit/Private/SFoliageEdit.h
--- a/Engine/Source/Editor/FoliageEdit/Private/SFoliageEdit.h
+++ b/Engine/Source/Editor/FoliageEdit/Private/SFoliageEdit.h
@@ -18,6 +18,9 @@ public:
 	/** Does a full refresh on the list. */
 	void RefreshFullList();
 
+	/** Notifies the widget that the mesh assigned to a foliage type in the list has changed */
+	void NotifyFoliageTypeMeshChanged(UFoliageType* FoliageType);
+
 	/** Gets FoliageEditMode. Used by the cluster details to notify changes */
 	class FEdModeFoliage* GetFoliageEditMode() const { return FoliageEditMode; }
 
diff --git a/Engine/Source/Editor/FoliageEdit/Private/SFoliagePalette.cpp b/Engine/Source/Editor/FoliageEdit/Private/SFoliagePalette.cpp
--- a/Engine/Source/Editor/FoliageEdit/Private/SFoliagePalette.cpp
+++ b/Engine/Source/Editor/FoliageEdit/Private/SFoliagePalette.cpp
@@ -93,6 +93,9 @@ private:
 ////////////////////////////////////////////////
 void SFoliagePalette::Construct(const FArguments& InArgs)
 {
+	bItemsNeedRebuild = false;
+	bIsActiveTimerRegistered = false;
+
 	FoliageEditMode = InArgs._FoliageEdMode;
 
 	FoliageEditMode->OnToolChanged.AddSP(this, &SFoliagePalette::HandleOnToolChanged);
@@ -176,6 +179,23 @@ void SFoliagePalette::Construct(const FArguments& InArgs)
 					.HintText(LOCTEXT("SearchFoliagePaletteHint", "Search Foliage"))
 					.OnTextChanged(this, &SFoliagePalette::OnSearchTextChanged)
 				]
+
+				// View Options
+				+ SHorizontalBox::Slot()
+				.HAlign(HAlign_Right)
+				.AutoWidth()
+				[
+					SNew( SComboButton )
+					.ContentPadding(0)
+					.ForegroundColor( FSlateColor::UseForeground() )
+					.ButtonStyle( FEditorStyle::Get(), "ToggleButton" )
+					.OnGetMenuContent(this, &SFoliagePalette::GetViewOptionsMenuContent)
+					.ButtonContent()
+					[
+						SNew(SImage)
+						.Image( FEditorStyle::GetBrush("GenericViewButton") )
+					]
+				]
 			]
 		]
 
@@ -224,7 +244,17 @@ void SFoliagePalette::Construct(const FArguments& InArgs)
 					.AutoHeight()
 					[
 						SNew(SHorizontalBox)
-					
+						
+						// Selected type name area
+						+ SHorizontalBox::Slot()
+						.Padding(FMargin(3.f))
+						.VAlign(VAlign_Bottom)
+						//.AutoWidth()
+						[
+							SNew(STextBlock)
+							.Text(this, &SFoliagePalette::GetDetailsNameAreaText)
+						]
+
 						// Show/Hide Details
 						+ SHorizontalBox::Slot()
 						.HAlign(HAlign_Right)
@@ -240,55 +270,28 @@ void SFoliagePalette::Construct(const FArguments& InArgs)
 							[
 								SNew(SHorizontalBox)
 
-								// Arrow
-								+ SHorizontalBox::Slot()
+								// Details icon
+								+SHorizontalBox::Slot()
 								.AutoWidth()
 								.HAlign(HAlign_Center)
 								.VAlign(VAlign_Center)
 								[
 									SNew(SImage)
-									.Image(this, &SFoliagePalette::GetShowHideDetailsImage)
+									.Image(FEditorStyle::GetBrush("LevelEditor.Tabs.Details"))
 								]
 
-								// Details icon
-								+SHorizontalBox::Slot()
-								.AutoWidth()
+								// Arrow
+								+ SHorizontalBox::Slot()
 								.Padding(FMargin(3.f, 0.f))
+								.AutoWidth()
 								.HAlign(HAlign_Center)
 								.VAlign(VAlign_Center)
 								[
 									SNew(SImage)
-									.Image(FEditorStyle::GetBrush("LevelEditor.Tabs.Details"))
+									.Image(this, &SFoliagePalette::GetShowHideDetailsImage)
 								]
 							]
 						]
-
-						// Selected type name area
-						+ SHorizontalBox::Slot()
-						.Padding(FMargin(3.f))
-						.VAlign(VAlign_Bottom)
-						//.AutoWidth()
-						[
-							SNew(STextBlock)
-							.Text(this, &SFoliagePalette::GetDetailsNameAreaText)
-						]
-
-						// View Options
-						+ SHorizontalBox::Slot()
-						.HAlign(HAlign_Right)
-						.AutoWidth()
-						[
-							SNew( SComboButton )
-							.ContentPadding(0)
-							.ForegroundColor( FSlateColor::UseForeground() )
-							.ButtonStyle( FEditorStyle::Get(), "ToggleButton" )
-							.OnGetMenuContent(this, &SFoliagePalette::GetViewOptionsMenuContent)
-							.ButtonContent()
-							[
-								SNew(SImage)
-								.Image( FEditorStyle::GetBrush("GenericViewButton") )
-							]
-						]
 					]
 				]
 				
@@ -333,9 +336,9 @@ FReply SFoliagePalette::OnKeyDown(const FGeometry& MyGeometry, const FKeyEvent&
 	return SCompoundWidget::OnKeyDown(MyGeometry, InKeyEvent);
 }
 
-void SFoliagePalette::UpdatePalette(bool bRefreshItems)
+void SFoliagePalette::UpdatePalette(bool bRebuildItems)
 {
-	bItemsNeedRefresh |= bRefreshItems;
+	bItemsNeedRebuild |= bRebuildItems;
 
 	if (!bIsActiveTimerRegistered)
 	{
@@ -344,6 +347,27 @@ void SFoliagePalette::UpdatePalette(bool bRefreshItems)
 	}
 }
 
+void SFoliagePalette::UpdateThumbnailForType(UFoliageType* FoliageType)
+{
+	// Recreate the palette item for the given foliage type
+	for (auto& Item : PaletteItems)
+	{
+		if (Item->GetFoliageType() == FoliageType)
+		{
+			const bool bItemIsSelected = GetActiveViewWidget()->IsItemSelected(Item);
+
+			Item = MakeShareable(new FFoliagePaletteItemModel(Item->GetTypeUIInfo(), SharedThis(this), FoliageEditMode));
+			if (bItemIsSelected)
+			{
+				GetActiveViewWidget()->SetItemSelection(Item, true);
+			}
+
+			UpdatePalette();
+			break;
+		}
+	}
+}
+
 bool SFoliagePalette::AnySelectedTileHovered() const
 {
 	for (auto& TypeInfo : GetActiveViewWidget()->GetSelectedItems())
@@ -372,6 +396,20 @@ void SFoliagePalette::BindCommands()
 	const FFoliagePaletteCommands& Commands = FFoliagePaletteCommands::Get();
 
 	// Context menu commands
+	UICommandList->MapAction(
+		Commands.ActivateFoliageType,
+		FExecuteAction::CreateSP(this, &SFoliagePalette::OnActivateFoliageTypes),
+		FCanExecuteAction(),
+		FIsActionChecked(),
+		FIsActionButtonVisible::CreateSP(this, &SFoliagePalette::OnCanActivateFoliageTypes));
+
+	UICommandList->MapAction(
+		Commands.DeactivateFoliageType,
+		FExecuteAction::CreateSP(this, &SFoliagePalette::OnDeactivateFoliageTypes),
+		FCanExecuteAction(),
+		FIsActionChecked(),
+		FIsActionButtonVisible::CreateSP(this, &SFoliagePalette::OnCanDeactivateFoliageTypes));
+
 	UICommandList->MapAction(
 		Commands.RemoveFoliageType,
 		FExecuteAction::CreateSP(this, &SFoliagePalette::OnRemoveFoliageType));
@@ -419,11 +457,7 @@ void SFoliagePalette::AddFoliageType(const FAssetData& AssetData)
 	UObject* Asset = AssetData.GetAsset();
 	GWarn->EndSlowTask();
 
-	UFoliageType* FoliageType = FoliageEditMode->AddFoliageAsset(Asset);
-	if (FoliageType)
-	{
-		UpdatePalette(true);
-	}
+	FoliageEditMode->AddFoliageAsset(Asset);
 }
 
 TSharedRef<SWidgetSwitcher> SFoliagePalette::CreatePaletteViews()
@@ -474,13 +508,14 @@ TSharedRef<SWidgetSwitcher> SFoliagePalette::CreatePaletteViews()
 		.SortMode(this, &SFoliagePalette::GetMeshColumnSortMode)
 		.OnSort(this, &SFoliagePalette::OnMeshesColumnSortModeChanged)
 		.DefaultLabel(this, &SFoliagePalette::GetMeshesHeaderText)
+		.FillWidth(5.f)
 
 		// Instance Count
 		+ SHeaderRow::Column(FoliagePaletteTreeColumns::ColumnID_InstanceCount)
 		.HeaderContentPadding(FMargin(10, 1, 0, 1))
 		.DefaultLabel(LOCTEXT("InstanceCount", "Count"))
 		.DefaultTooltip(this, &SFoliagePalette::GetTotalInstanceCountTooltipText)
-		.FixedWidth(75.f)
+		.FillWidth(2.f)
 
 		// Save Asset
 		+ SHeaderRow::Column(FoliagePaletteTreeColumns::ColumnID_Save)
@@ -608,8 +643,7 @@ void SFoliagePalette::OnSelectionChanged(FFoliagePaletteItemModelPtr Item, ESele
 
 void SFoliagePalette::OnItemDoubleClicked(FFoliagePaletteItemModelPtr Item) const
 {
-	const bool bCurrentlyActive = Item->GetFoliageType()->IsSelected;
-	Item->SetTypeActiveInPalette(!bCurrentlyActive);
+	Item->SetTypeActiveInPalette(!Item->IsActive());
 }
 
 TSharedRef<SWidget> SFoliagePalette::GetViewOptionsMenuContent()
@@ -741,9 +775,11 @@ TSharedPtr<SWidget> SFoliagePalette::ConstructFoliageTypeContextMenu()
 	const FFoliagePaletteCommands& Commands = FFoliagePaletteCommands::Get();
 	FMenuBuilder MenuBuilder(true, UICommandList);
 
-	if (GetActiveViewWidget()->GetSelectedItems().Num() > 0)
+	auto SelectedItems = GetActiveViewWidget()->GetSelectedItems();
+	if (SelectedItems.Num() > 0)
 	{
-		if (AreAnyNonAssetTypesSelected())
+		const bool bShowSaveAsOption = SelectedItems.Num() == 1 && !SelectedItems[0]->GetFoliageType()->IsAsset();
+		if (bShowSaveAsOption)
 		{
 			MenuBuilder.BeginSection("StaticMeshFoliageTypeOptions", LOCTEXT("StaticMeshFoliageTypeOptionsHeader", "Static Mesh"));
 			{
@@ -752,17 +788,34 @@ TSharedPtr<SWidget> SFoliagePalette::ConstructFoliageTypeContextMenu()
 					LOCTEXT("SaveAsFoliageType_ToolTip", "Creates a Foliage Type asset with these settings that can be reused in other levels."),
 					FSlateIcon(FEditorStyle::GetStyleSetName(), "Level.SaveIcon16x"),
 					FUIAction(
-						FExecuteAction::CreateSP(this, &SFoliagePalette::OnSaveSelectedAsFoliageType),
-						FCanExecuteAction::CreateSP(this, &SFoliagePalette::OnCanSaveSelectedAsFoliageType)
-					),
+						FExecuteAction::CreateSP(this, &SFoliagePalette::OnSaveSelected)
+						),
 					NAME_None
-				);
+					);
 			}
 			MenuBuilder.EndSection();
 		}
 
 		MenuBuilder.BeginSection("FoliageTypeOptions", LOCTEXT("FoliageTypeOptionsHeader", "Foliage Type"));
 		{
+			if (!bShowSaveAsOption)
+			{
+				MenuBuilder.AddMenuEntry(
+					LOCTEXT("SaveSelectedFoliageTypes", "Save"),
+					LOCTEXT("SaveSelectedFoliageTypes_ToolTip", "Saves any changes to the selected foliage type asset(s)."),
+					FSlateIcon(FEditorStyle::GetStyleSetName(), "Level.SaveIcon16x"),
+					FUIAction(
+						FExecuteAction::CreateSP(this, &SFoliagePalette::OnSaveSelected),
+						FCanExecuteAction::CreateSP(this, &SFoliagePalette::OnCanSaveAnySelectedAssets)
+						),
+					NAME_None
+					);
+			}
+
+			MenuBuilder.AddMenuEntry(Commands.ActivateFoliageType);
+
+			MenuBuilder.AddMenuEntry(Commands.DeactivateFoliageType);
+
 			MenuBuilder.AddMenuEntry(Commands.RemoveFoliageType);
 
 			MenuBuilder.AddSubMenu(
@@ -786,31 +839,35 @@ TSharedPtr<SWidget> SFoliagePalette::ConstructFoliageTypeContextMenu()
 	return MenuBuilder.MakeWidget();
 }
 
-void SFoliagePalette::OnSaveSelectedAsFoliageType()
+void SFoliagePalette::OnSaveSelected()
 {
 	for (FFoliagePaletteItemModelPtr& PaletteItem : GetActiveViewWidget()->GetSelectedItems())
 	{
 		UFoliageType* FoliageType = PaletteItem->GetFoliageType();
-		UFoliageType* SavedFoliageType = FoliageEditMode->SaveFoliageTypeObject(FoliageType);
-		if (SavedFoliageType)
+		if (!FoliageType->IsAsset() || FoliageType->GetOutermost()->IsDirty())
 		{
-			FoliageType = SavedFoliageType;
+			UFoliageType* SavedFoliageType = FoliageEditMode->SaveFoliageTypeObject(FoliageType);
+			if (SavedFoliageType)
+			{
+				FoliageType = SavedFoliageType;
+			}
 		}
 	}
 }
 
-bool SFoliagePalette::OnCanSaveSelectedAsFoliageType() const
+bool SFoliagePalette::OnCanSaveAnySelectedAssets() const
 {
+	// We can save if at least one of the selected items is a dirty asset
 	for (FFoliagePaletteItemModelPtr& PaletteItem : GetActiveViewWidget()->GetSelectedItems())
 	{
-		if (PaletteItem->GetFoliageType()->IsAsset())
+		UFoliageType* FoliageType = PaletteItem->GetFoliageType();
+		if (FoliageType->IsAsset() && FoliageType->GetOutermost()->IsDirty())
 		{
-			// At least one selected type is an asset
-			return false;
+			return true;
 		}
 	}
 
-	return true;
+	return false;
 }
 
 bool SFoliagePalette::AreAnyNonAssetTypesSelected() const
@@ -827,6 +884,50 @@ bool SFoliagePalette::AreAnyNonAssetTypesSelected() const
 	return false;
 }
 
+void SFoliagePalette::OnActivateFoliageTypes()
+{
+	for (FFoliagePaletteItemModelPtr& PaletteItem : GetActiveViewWidget()->GetSelectedItems())
+	{
+		PaletteItem->SetTypeActiveInPalette(true);
+	}
+}
+
+bool SFoliagePalette::OnCanActivateFoliageTypes() const
+{
+	// At least one selected item must be inactive
+	for (FFoliagePaletteItemModelPtr& PaletteItem : GetActiveViewWidget()->GetSelectedItems())
+	{
+		if (!PaletteItem->IsActive())
+		{
+			return true;
+		}
+	}
+
+	return false;
+}
+
+void SFoliagePalette::OnDeactivateFoliageTypes()
+{
+	for (FFoliagePaletteItemModelPtr& PaletteItem : GetActiveViewWidget()->GetSelectedItems())
+	{
+		PaletteItem->SetTypeActiveInPalette(false);
+	}
+}
+
+bool SFoliagePalette::OnCanDeactivateFoliageTypes() const
+{
+	// At least one selected item must be active
+	for (FFoliagePaletteItemModelPtr& PaletteItem : GetActiveViewWidget()->GetSelectedItems())
+	{
+		if (PaletteItem->IsActive())
+		{
+			return true;
+		}
+	}
+
+	return false;
+}
+
 void SFoliagePalette::FillReplaceFoliageTypeSubmenu(FMenuBuilder& MenuBuilder)
 {
 	FContentBrowserModule& ContentBrowserModule = FModuleManager::Get().LoadModuleChecked<FContentBrowserModule>(TEXT("ContentBrowser"));
@@ -990,8 +1091,7 @@ ECheckBoxState SFoliagePalette::GetState_AllMeshes() const
 
 	for (FFoliagePaletteItemModelPtr& PaletteItem : GetActiveViewWidget()->GetSelectedItems())
 	{
-		UFoliageType* FoliageType = PaletteItem->GetFoliageType();
-		if (FoliageType->IsSelected)
+		if (PaletteItem->IsActive())
 		{
 			bHasChecked = true;
 		}
@@ -1011,10 +1111,10 @@ ECheckBoxState SFoliagePalette::GetState_AllMeshes() const
 
 void SFoliagePalette::OnCheckStateChanged_AllMeshes(ECheckBoxState InState)
 {
+	const bool bActivate = InState == ECheckBoxState::Checked;
 	for (FFoliagePaletteItemModelPtr& PaletteItem : GetActiveViewWidget()->GetSelectedItems())
 	{
-		UFoliageType* FoliageType = PaletteItem->GetFoliageType();
-		FoliageType->IsSelected = (InState == ECheckBoxState::Checked);
+		PaletteItem->SetTypeActiveInPalette(bActivate);
 	}
 }
 
@@ -1110,44 +1210,46 @@ FReply SFoliagePalette::OnShowHideDetailsClicked() const
 
 EActiveTimerReturnType SFoliagePalette::UpdatePaletteItems(double InCurrentTime, float InDeltaTime)
 {
-	if (bItemsNeedRefresh)
+	if (bItemsNeedRebuild)
 	{
-		const auto& IFATypeList = FoliageEditMode->GetFoliageMeshList();
+		bItemsNeedRebuild = false;
 
-		// Check if any palette items should be removed
-		// We do this instead of a full rebuild to avoid the flicker that occurs when creating a new thumbnail widget
-		TArray<FFoliageMeshUIInfoPtr> ItemsToRemove;
-		for (auto& ItemPair : PaletteItemsByTypeInfo)
-		{
-			if (!IFATypeList.Contains(ItemPair.Key))
-			{
-				ItemsToRemove.Add(ItemPair.Key);
-			}
-		}
+		// Cache the currently selected items
+		auto ActiveViewWidget = GetActiveViewWidget();
+		TArray<FFoliagePaletteItemModelPtr> PreviouslySelectedItems = ActiveViewWidget->GetSelectedItems();
+		
+		ActiveViewWidget->ClearSelection();
 
-		// Remove the items from the palette
-		for (FFoliageMeshUIInfoPtr& Item : ItemsToRemove)
+		// Rebuild the list of palette items
+		const auto& AllTypesList = FoliageEditMode->GetFoliageMeshList();
+		PaletteItems.Empty(AllTypesList.Num());
+		for (const FFoliageMeshUIInfoPtr& TypeInfo : AllTypesList)
 		{
-			PaletteItemsByTypeInfo.Remove(Item);
+			PaletteItems.Add(MakeShareable(new FFoliagePaletteItemModel(TypeInfo, SharedThis(this), FoliageEditMode)));
 		}
 
-		// Create a palette item for any newly added types
-		for (FFoliageMeshUIInfoPtr& TypeInfo : FoliageEditMode->GetFoliageMeshList())
+		// Restore the selection
+		for (auto& PrevSelectedItem : PreviouslySelectedItems)
 		{
-			if (!PaletteItemsByTypeInfo.Contains(TypeInfo))
+			// See if there's a new item for the previously selected item's foliage type
+			for (auto& Item : PaletteItems)
 			{
-				PaletteItemsByTypeInfo.Add(TypeInfo) = MakeShareable(new FFoliagePaletteItemModel(TypeInfo, SharedThis(this), FoliageEditMode));
+				if (PrevSelectedItem->GetFoliageType() == Item->GetFoliageType())
+				{
+					ActiveViewWidget->SetItemSelection(Item, true);
+					break;
+				}
 			}
 		}
 	}
 
 	// Update the filtered items
 	FilteredItems.Empty();
-	for (auto& ItemPair : PaletteItemsByTypeInfo)
+	for (auto& Item : PaletteItems)
 	{
-		if (TypeFilter->PassesFilter(ItemPair.Value))
+		if (TypeFilter->PassesFilter(Item))
 		{
-			FilteredItems.Add(ItemPair.Value);
+			FilteredItems.Add(Item);
 		}
 	}
 
diff --git a/Engine/Source/Editor/FoliageEdit/Private/SFoliagePalette.h b/Engine/Source/Editor/FoliageEdit/Private/SFoliagePalette.h
--- a/Engine/Source/Editor/FoliageEdit/Private/SFoliagePalette.h
+++ b/Engine/Source/Editor/FoliageEdit/Private/SFoliagePalette.h
@@ -29,8 +29,11 @@ public:
 	virtual FReply OnKeyDown(const FGeometry& MyGeometry, const FKeyEvent& InKeyEvent) override;
 
 	/** Updates the foliage palette, optionally doing a full refresh of the items in the palette as well */
-	void UpdatePalette(bool bRefreshItems = false);
-	
+	void UpdatePalette(bool bRebuildItems = false);
+
+	/** Updates the thumbnail for the given foliage type in the palette */
+	void UpdateThumbnailForType(UFoliageType* FoliageType);
+
 	bool AnySelectedTileHovered() const;
 	void ActivateAllSelectedTypes(bool bActivate) const;
 
@@ -114,14 +117,22 @@ private:	// CONTEXT MENU
 	TSharedPtr<SWidget> ConstructFoliageTypeContextMenu();
 
 	/** Saves a temporary non-asset foliage type (created from a static mesh) as a foliage type asset */
-	void OnSaveSelectedAsFoliageType();
+	void OnSaveSelected();
 
 	/** @return True if all selected types are non-assets */
-	bool OnCanSaveSelectedAsFoliageType() const;
+	bool OnCanSaveAnySelectedAssets() const;
 
 	/** @return True if any of the selected types are non-assets */
 	bool AreAnyNonAssetTypesSelected() const;
 
+	/** Handler for the 'Activate' command */
+	void OnActivateFoliageTypes();
+	bool OnCanActivateFoliageTypes() const;
+
+	/** Handler for the 'Deactivate' command */
+	void OnDeactivateFoliageTypes();
+	bool OnCanDeactivateFoliageTypes() const;
+
 	/** Fills 'Replace' menu command  */
 	void FillReplaceFoliageTypeSubmenu(FMenuBuilder& MenuBuilder);
 
@@ -206,8 +217,8 @@ private:
 	typedef TTextFilter<FFoliagePaletteItemModelPtr> FoliageTypeTextFilter;
 	TSharedPtr<FoliageTypeTextFilter> TypeFilter;
 
-	/** All the items in the palette */
-	TMap<FFoliageMeshUIInfoPtr, FFoliagePaletteItemModelPtr> PaletteItemsByTypeInfo;
+	/** All the items in the palette (unfiltered) */
+	TArray<FFoliagePaletteItemModelPtr> PaletteItems;
 
 	/** The filtered list of types to display in the palette */
 	TArray<FFoliagePaletteItemModelPtr> FilteredItems;
@@ -235,6 +246,6 @@ private:
 
 	FEdModeFoliage* FoliageEditMode;
 
-	bool bItemsNeedRefresh : 1;
+	bool bItemsNeedRebuild : 1;
 	bool bIsActiveTimerRegistered : 1;
 };
\ No newline at end of file
diff --git a/Engine/Source/Runtime/Foliage/Private/InstancedFoliage.cpp b/Engine/Source/Runtime/Foliage/Private/InstancedFoliage.cpp
--- a/Engine/Source/Runtime/Foliage/Private/InstancedFoliage.cpp
+++ b/Engine/Source/Runtime/Foliage/Private/InstancedFoliage.cpp
@@ -485,12 +485,16 @@ void UFoliageType::PostEditChangeProperty(struct FPropertyChangedEvent& Property
 
 	UpdateGuid = FGuid::NewGuid();
 
+	//@todo: move this into FoliageType_InstancedStaticMesh
+	// Check to see if the mesh is what changed
+	const bool bMeshChanged = PropertyChangedEvent.Property && PropertyChangedEvent.Property->GetFName() == "Mesh";
+
 	// Notify any currently-loaded InstancedFoliageActors
 	if (IsFoliageReallocationRequiredForPropertyChange(PropertyChangedEvent))
 	{
 		for (TObjectIterator<AInstancedFoliageActor> It(RF_ClassDefaultObject | RF_PendingKill); It; ++It)
 		{
-			It->NotifyFoliageTypeChanged(this);
+			It->NotifyFoliageTypeChanged(this, bMeshChanged);
 		}
 	}
 }
@@ -587,6 +591,13 @@ void FFoliageMeshInfo::UpdateComponentSettings(const UFoliageType* InSettings)
 	{
 		bool bNeedsMarkRenderStateDirty = false;
 		bool bNeedsInvalidateLightingCache = false;
+		
+		if (Component->StaticMesh != InSettings->GetStaticMesh())
+		{
+			Component->StaticMesh = InSettings->GetStaticMesh();
+			bNeedsInvalidateLightingCache = true;
+			bNeedsMarkRenderStateDirty = true;
+		}
 
 		EComponentMobility::Type NewMobility = InSettings->bEnableStaticLighting ? EComponentMobility::Static : EComponentMobility::Movable;
 		if (Component->Mobility != NewMobility)
@@ -2228,12 +2239,19 @@ void AInstancedFoliageActor::PostLoad()
 
 #if WITH_EDITOR
 
-void AInstancedFoliageActor::NotifyFoliageTypeChanged(UFoliageType* FoliageType)
+void AInstancedFoliageActor::NotifyFoliageTypeChanged(UFoliageType* FoliageType, bool bMeshChanged)
 {
-	FFoliageMeshInfo* MeshInfo = FindMesh(FoliageType);
-	if (MeshInfo)
+	FFoliageMeshInfo* TypeInfo = FindMesh(FoliageType);
+
+	if (TypeInfo)
 	{
-		MeshInfo->UpdateComponentSettings(FoliageType);
+		TypeInfo->UpdateComponentSettings(FoliageType);
+
+		if (bMeshChanged)
+		{
+			// If the type's mesh has changed, the UI needs to be notified so it can update thumbnails accordingly
+			OnFoliageTypeMeshChangedEvent.Broadcast(FoliageType);
+		}
 	}
 }
 
diff --git a/Engine/Source/Runtime/Foliage/Public/InstancedFoliageActor.h b/Engine/Source/Runtime/Foliage/Public/InstancedFoliageActor.h
--- a/Engine/Source/Runtime/Foliage/Public/InstancedFoliageActor.h
+++ b/Engine/Source/Runtime/Foliage/Public/InstancedFoliageActor.h
@@ -157,7 +157,10 @@ public:
 	static FOLIAGE_API bool HasFoliageAttached(UActorComponent* InComponent);
 
 	/* Called to notify InstancedFoliageActor that a UFoliageType has been modified */
-	void NotifyFoliageTypeChanged(UFoliageType* FoliageType);
+	void NotifyFoliageTypeChanged(UFoliageType* FoliageType, bool bMeshChanged);
+
+	DECLARE_EVENT_OneParam(AInstancedFoliageActor, FOnFoliageTypeMeshChanged, UFoliageType*);
+	FOnFoliageTypeMeshChanged& OnFoliageTypeMeshChanged() { return OnFoliageTypeMeshChangedEvent; }
 
 #endif	//WITH_EDITOR
 
@@ -177,6 +180,8 @@ private:
 	FDelegateHandle OnLevelActorMovedDelegateHandle;
 	FDelegateHandle OnLevelActorDeletedDelegateHandle;
 	FDelegateHandle OnPostApplyLevelOffsetDelegateHandle;
+
+	FOnFoliageTypeMeshChanged OnFoliageTypeMeshChangedEvent;
 #endif
 
 };
