[UE-15100] - Re-adding resimulate tooltips and limitations

diff --git a/Engine/Source/Editor/FoliageEdit/Private/ProceduralFoliageComponentDetails.cpp b/Engine/Source/Editor/FoliageEdit/Private/ProceduralFoliageComponentDetails.cpp
--- a/Engine/Source/Editor/FoliageEdit/Private/ProceduralFoliageComponentDetails.cpp
+++ b/Engine/Source/Editor/FoliageEdit/Private/ProceduralFoliageComponentDetails.cpp
@@ -11,29 +11,30 @@
 #include "SNotificationList.h"
 #include "NotificationManager.h"
 
+#define LOCTEXT_NAMESPACE "ProceduralFoliageComponentDetails"
+
 TSharedRef<IDetailCustomization> FProceduralFoliageComponentDetails::MakeInstance()
 {
-	return MakeShareable( new FProceduralFoliageComponentDetails() );
+	return MakeShareable(new FProceduralFoliageComponentDetails());
 
 }
 
-
-void FProceduralFoliageComponentDetails::CustomizeDetails( IDetailLayoutBuilder& DetailBuilder )
+void FProceduralFoliageComponentDetails::CustomizeDetails(IDetailLayoutBuilder& DetailBuilder)
 {
 	const FName ProceduralFoliageCategoryName("ProceduralFoliage");
 	IDetailCategoryBuilder& ProceduralFoliageCategory = DetailBuilder.EditCategory(ProceduralFoliageCategoryName);
 
-	const FText ResimulateText = NSLOCTEXT("ProceduralFoliageComponentDetails","ResimulateButtonText", "Resimulate" );
+	const FText ResimulateText = LOCTEXT("ResimulateButtonText", "Resimulate");
 
 	TArray< TWeakObjectPtr<UObject> > ObjectsBeingCustomized;
-	DetailBuilder.GetObjectsBeingCustomized( ObjectsBeingCustomized );
+	DetailBuilder.GetObjectsBeingCustomized(ObjectsBeingCustomized);
 
-	for( TWeakObjectPtr<UObject>& Object : ObjectsBeingCustomized )
+	for (TWeakObjectPtr<UObject>& Object : ObjectsBeingCustomized)
 	{
-		UProceduralFoliageComponent* Component = Cast<UProceduralFoliageComponent>( Object.Get() );
-		if( ensure( Component ) )
+		UProceduralFoliageComponent* Component = Cast<UProceduralFoliageComponent>(Object.Get());
+		if (ensure(Component))
 		{
-			SelectedComponents.Add( Component );
+			SelectedComponents.Add(Component);
 		}
 	}
 
@@ -42,50 +43,50 @@ void FProceduralFoliageComponentDetails::CustomizeDetails( IDetailLayoutBuilder&
 	bool bAdvancedProperties = false;
 	// Add all properties in the category in order
 	ProceduralFoliageCategory.GetDefaultProperties(AllProperties, bSimpleProperties, bAdvancedProperties);
-	for( auto& Property : AllProperties )
+	for (auto& Property : AllProperties)
 	{
 		ProceduralFoliageCategory.AddProperty(Property);
 	}
 
-	FDetailWidgetRow& NewRow =  ProceduralFoliageCategory.AddCustomRow( ResimulateText );
+	FDetailWidgetRow& NewRow = ProceduralFoliageCategory.AddCustomRow(ResimulateText);
 
 	NewRow.ValueContent()
-	.MaxDesiredWidth(120.f)
-	[
-		SNew(SButton)
-		.OnClicked( this, &FProceduralFoliageComponentDetails::OnResimulateClicked )
-		.ToolTipText( NSLOCTEXT("ProceduralFoliageComponentDetails","ResimulateButton_Tooltip", "Runs the procedural foliage spawner simulation. Replaces any existing instances spawned by a previous simulation." ) )
-		.IsEnabled( this, &FProceduralFoliageComponentDetails::IsResimulateEnabled )
+		.MaxDesiredWidth(120.f)
 		[
-			SNew( STextBlock )
-			.Font( IDetailLayoutBuilder::GetDetailFont() )
-			.Text( ResimulateText )
-		]
-	];
+			SNew(SButton)
+			.OnClicked(this, &FProceduralFoliageComponentDetails::OnResimulateClicked)
+			.ToolTipText(this, &FProceduralFoliageComponentDetails::GetResimulateTooltipText)
+			.IsEnabled(this, &FProceduralFoliageComponentDetails::IsResimulateEnabled)
+			[
+				SNew(STextBlock)
+				.Font(IDetailLayoutBuilder::GetDetailFont())
+				.Text(ResimulateText)
+			]
+		];
 }
 
 FReply FProceduralFoliageComponentDetails::OnResimulateClicked()
 {
 	TSet<UProceduralFoliageSpawner*> UniqueFoliageSpawners;
-	for( TWeakObjectPtr<UProceduralFoliageComponent>& Component : SelectedComponents )
+	for (TWeakObjectPtr<UProceduralFoliageComponent>& Component : SelectedComponents)
 	{
-		if( Component.IsValid() && Component->FoliageSpawner )
+		if (Component.IsValid() && Component->FoliageSpawner)
 		{
-			if( !UniqueFoliageSpawners.Contains( Component->FoliageSpawner ) )
+			if (!UniqueFoliageSpawners.Contains(Component->FoliageSpawner))
 			{
 				UniqueFoliageSpawners.Add(Component->FoliageSpawner);
 			}
 
-			FScopedTransaction Transaction(NSLOCTEXT("ProceduralFoliageComponentDetails", "Resimulate_Transaction", "Procedural Foliage Simulation"));
+			FScopedTransaction Transaction(LOCTEXT("Resimulate_Transaction", "Procedural Foliage Simulation"));
 			TArray <FDesiredFoliageInstance> DesiredFoliageInstances;
 			if (Component->GenerateProceduralContent(DesiredFoliageInstances))
 			{
 				FEdModeFoliage::AddInstances(Component->GetWorld(), DesiredFoliageInstances);
 
-				// If no instance was actually spawned, inform the user
+				// If no instances were spawned, inform the user
 				if (!Component->HasSpawnedAnyInstances())
 				{
-					FNotificationInfo Info(NSLOCTEXT("ProceduralFoliageComponentDetails", "NothingSpawned_Notification", "Unable to spawn instances. Ensure a large enough surface exists within the volume."));
+					FNotificationInfo Info(LOCTEXT("NothingSpawned_Notification", "Unable to spawn instances. Ensure a large enough surface exists within the volume."));
 					Info.bUseLargeFont = false;
 					Info.bFireAndForget = true;
 					Info.bUseThrobber = false;
@@ -100,13 +101,70 @@ FReply FProceduralFoliageComponentDetails::OnResimulateClicked()
 
 bool FProceduralFoliageComponentDetails::IsResimulateEnabled() const
 {
+	bool bCanSimulate = false;
+
 	for(const TWeakObjectPtr<UProceduralFoliageComponent>& Component : SelectedComponents)
 	{
-		if(Component.IsValid() && Component->FoliageSpawner && Component->FoliageSpawner->GetFoliageTypes().Num() > 0)
+		if(Component.IsValid() && Component->FoliageSpawner)
 		{
-			return true;
+			for (const FFoliageTypeObject& FoliageTypeObject : Component->FoliageSpawner->GetFoliageTypes())
+			{
+				// Make sure at least one foliage type is ready to spawn
+				if (FoliageTypeObject.HasFoliageType())
+				{
+					bCanSimulate = true;
+					break;
+				}
+			}
+
+			if (bCanSimulate)
+			{
+				break;
+			}
 		}
 	}
 
-	return false;
-}
\ No newline at end of file
+	return bCanSimulate;
+}
+
+FText FProceduralFoliageComponentDetails::GetResimulateTooltipText() const
+{
+	FText TooltipText;
+
+	for (const TWeakObjectPtr<UProceduralFoliageComponent>& Component : SelectedComponents)
+	{
+		if (Component.IsValid())
+		{
+			if (!Component->FoliageSpawner)
+			{
+				TooltipText = LOCTEXT("Resimulate_Tooltip_NeedSpawner", "Cannot generate foliage: Assign a Procedural Foliage Spawner to run the procedural foliage simulation");
+			}
+			else
+			{
+				for (const FFoliageTypeObject& FoliageTypeObject : Component->FoliageSpawner->GetFoliageTypes())
+				{
+					// Make sure at least one foliage type is ready to spawn
+					if (!FoliageTypeObject.HasFoliageType())
+					{
+						TooltipText = LOCTEXT("Resimulate_Tooltip_EmptySpawner", "Cannot generate foliage: The assigned Procedural Foliage Spawner does not contain any foliage types to spawn.");
+						break;
+					}
+				}
+			}
+
+			if (!TooltipText.IsEmpty())
+			{
+				break;
+			}
+		}
+	}
+
+	if (TooltipText.IsEmpty())
+	{
+		TooltipText = LOCTEXT("Resimulate_Tooltip", "Runs the procedural foliage spawner simulation. Replaces any existing instances spawned by a previous simulation.");
+	}
+
+	return TooltipText;
+}
+
+#undef LOCTEXT_NAMESPACE
diff --git a/Engine/Source/Editor/FoliageEdit/Private/ProceduralFoliageComponentDetails.h b/Engine/Source/Editor/FoliageEdit/Private/ProceduralFoliageComponentDetails.h
--- a/Engine/Source/Editor/FoliageEdit/Private/ProceduralFoliageComponentDetails.h
+++ b/Engine/Source/Editor/FoliageEdit/Private/ProceduralFoliageComponentDetails.h
@@ -15,6 +15,7 @@ public:
 private:
 	FReply OnResimulateClicked();
 	bool IsResimulateEnabled() const;
+	FText GetResimulateTooltipText() const;
 private:
 	TArray< TWeakObjectPtr<class UProceduralFoliageComponent> > SelectedComponents;
 
diff --git a/Engine/Source/Runtime/Foliage/Private/FoliageTypeObject.cpp b/Engine/Source/Runtime/Foliage/Private/FoliageTypeObject.cpp
--- a/Engine/Source/Runtime/Foliage/Private/FoliageTypeObject.cpp
+++ b/Engine/Source/Runtime/Foliage/Private/FoliageTypeObject.cpp
@@ -40,11 +40,16 @@ const UFoliageType_InstancedStaticMesh* FFoliageTypeObject::GetInstance() const
 	return TypeInstance;
 }
 
-bool FFoliageTypeObject::IsValid() const
+bool FFoliageTypeObject::ContainsValidInstance() const
 {
 	return TypeInstance != nullptr;
 }
 
+bool FFoliageTypeObject::HasFoliageType() const
+{
+	return FoliageTypeObject != nullptr;
+}
+
 bool FFoliageTypeObject::IsDirty() const
 {
 	//@todo: make sure this works when the object isn't an asset
diff --git a/Engine/Source/Runtime/Foliage/Public/FoliageTypeObject.h b/Engine/Source/Runtime/Foliage/Public/FoliageTypeObject.h
--- a/Engine/Source/Runtime/Foliage/Public/FoliageTypeObject.h
+++ b/Engine/Source/Runtime/Foliage/Public/FoliageTypeObject.h
@@ -7,7 +7,7 @@
 
 /** A wrapper struct used to allow the use of either FoliageType assets or FoliageType blueprint classes */
 USTRUCT()
-struct FFoliageTypeObject
+struct FOLIAGE_API FFoliageTypeObject
 {
 	GENERATED_USTRUCT_BODY()
 
@@ -24,7 +24,10 @@ struct FFoliageTypeObject
 	const UFoliageType_InstancedStaticMesh* GetInstance() const;
 
 	/** @return Whether this would return a valid instance */
-	bool IsValid() const;
+	bool ContainsValidInstance() const;
+
+	/** @return Whether any foliage type is assigned at all */
+	bool HasFoliageType() const;
 
 	bool IsDirty() const;
 	void SetClean();
