[UE-15165] - Can once again delete selected foliage instances

diff --git a/Engine/Source/Editor/FoliageEdit/Private/FoliageEdMode.cpp b/Engine/Source/Editor/FoliageEdit/Private/FoliageEdMode.cpp
--- a/Engine/Source/Editor/FoliageEdit/Private/FoliageEdMode.cpp
+++ b/Engine/Source/Editor/FoliageEdit/Private/FoliageEdMode.cpp
@@ -2761,7 +2761,8 @@ bool FEdModeFoliage::InputKey(FEditorViewportClient* ViewportClient, FViewport*
 	}
 	
 	bool bHandled = false;
-	if (UISettings.GetPaintToolSelected() || UISettings.GetReapplyToolSelected() || UISettings.GetLassoSelectToolSelected())
+	const bool bInSelectionMode = UISettings.GetReapplyToolSelected() || UISettings.GetLassoSelectToolSelected();
+	if (UISettings.GetPaintToolSelected() || bInSelectionMode)
 	{
 		// Require Ctrl or not as per user preference
 		ELandscapeFoliageEditorControlType FoliageEditorControlType = GetDefault<ULevelEditorViewportSettings>()->FoliageEditorControlType;
@@ -2818,22 +2819,23 @@ bool FEdModeFoliage::InputKey(FEditorViewportClient* ViewportClient, FViewport*
 				bHandled = true;
 			}
 		}
-	}
-	else if (UISettings.GetSelectToolSelected() || UISettings.GetLassoSelectToolSelected())
-	{
-		if (Event == IE_Pressed)
+
+		if (!bHandled && bInSelectionMode)
 		{
-			if (Key == EKeys::Platform_Delete)
+			if (Event == IE_Pressed)
 			{
-				RemoveSelectedInstances(GetWorld());
-				
-				bHandled = true;
-			}
-			else if (Key == EKeys::End)
-			{
-				SnapSelectedInstancesToGround(GetWorld());
-				
-				bHandled = true;
+				if (Key == EKeys::Platform_Delete)
+				{
+					RemoveSelectedInstances(GetWorld());
+
+					bHandled = true;
+				}
+				else if (Key == EKeys::End)
+				{
+					SnapSelectedInstancesToGround(GetWorld());
+
+					bHandled = true;
+				}
 			}
 		}
 	}
