[UE-15434] - Can delete foliage instances when regular select tool is selected
- Fixed snapping selected foliage instances to ground to be an undo-able action

diff --git a/Engine/Source/Editor/FoliageEdit/Private/FoliageEdMode.cpp b/Engine/Source/Editor/FoliageEdit/Private/FoliageEdMode.cpp
--- a/Engine/Source/Editor/FoliageEdit/Private/FoliageEdMode.cpp
+++ b/Engine/Source/Editor/FoliageEdit/Private/FoliageEdMode.cpp
@@ -2243,48 +2243,52 @@ bool FEdModeFoliage::GetStaticMeshVertexColorForHit(const UStaticMeshComponent*
 
 void FEdModeFoliage::SnapSelectedInstancesToGround(UWorld* InWorld)
 {
-	bool bMovedInstance = false;
-	
-	const int32 NumLevels = InWorld->GetNumLevels();
-	for (int32 LevelIdx = 0; LevelIdx < NumLevels; ++LevelIdx)
+	GEditor->BeginTransaction(NSLOCTEXT("UnrealEd", "FoliageMode_EditTransaction", "Snap Foliage To Ground"));
 	{
-		ULevel* Level = InWorld->GetLevel(LevelIdx);
-		AInstancedFoliageActor* IFA = AInstancedFoliageActor::GetInstancedFoliageActorForLevel(Level);
-		if (IFA)
-		{
-			bool bFoundSelection = false;
+		bool bMovedInstance = false;
 
-			for (auto& MeshPair : IFA->FoliageMeshes)
+		const int32 NumLevels = InWorld->GetNumLevels();
+		for (int32 LevelIdx = 0; LevelIdx < NumLevels; ++LevelIdx)
+		{
+			ULevel* Level = InWorld->GetLevel(LevelIdx);
+			AInstancedFoliageActor* IFA = AInstancedFoliageActor::GetInstancedFoliageActorForLevel(Level);
+			if (IFA)
 			{
-				FFoliageMeshInfo& MeshInfo = *MeshPair.Value;
-				TArray<int32> SelectedIndices = MeshInfo.SelectedIndices.Array();
+				bool bFoundSelection = false;
 
-				if (SelectedIndices.Num() > 0)
+				for (auto& MeshPair : IFA->FoliageMeshes)
 				{
-					// Mark actor once we found selection
-					if (!bFoundSelection)
+					FFoliageMeshInfo& MeshInfo = *MeshPair.Value;
+					TArray<int32> SelectedIndices = MeshInfo.SelectedIndices.Array();
+
+					if (SelectedIndices.Num() > 0)
 					{
-						IFA->Modify();
-						bFoundSelection = true;
-					}
+						// Mark actor once we found selection
+						if (!bFoundSelection)
+						{
+							IFA->Modify();
+							bFoundSelection = true;
+						}
 
-					MeshInfo.PreMoveInstances(IFA, SelectedIndices);
+						MeshInfo.PreMoveInstances(IFA, SelectedIndices);
 
-					for (int32 InstanceIndex : SelectedIndices)
-					{
-						bMovedInstance|= SnapInstanceToGround(IFA, MeshPair.Key->AlignMaxAngle, MeshInfo, InstanceIndex);
-					}
+						for (int32 InstanceIndex : SelectedIndices)
+						{
+							bMovedInstance |= SnapInstanceToGround(IFA, MeshPair.Key->AlignMaxAngle, MeshInfo, InstanceIndex);
+						}
 
-					MeshInfo.PostMoveInstances(IFA, SelectedIndices);
+						MeshInfo.PostMoveInstances(IFA, SelectedIndices);
+					}
 				}
 			}
 		}
-	}
 
-	if (bMovedInstance)
-	{
-		UpdateWidgetLocationToInstanceSelection();
+		if (bMovedInstance)
+		{
+			UpdateWidgetLocationToInstanceSelection();
+		}
 	}
+	GEditor->EndTransaction();
 }
 
 void FEdModeFoliage::HandleOnActorSpawned(AActor* Actor)
@@ -2761,8 +2765,7 @@ bool FEdModeFoliage::InputKey(FEditorViewportClient* ViewportClient, FViewport*
 	}
 	
 	bool bHandled = false;
-	const bool bInSelectionMode = UISettings.GetReapplyToolSelected() || UISettings.GetLassoSelectToolSelected();
-	if (UISettings.GetPaintToolSelected() || bInSelectionMode)
+	if (UISettings.GetPaintToolSelected() || UISettings.GetReapplyToolSelected() || UISettings.GetLassoSelectToolSelected())
 	{
 		// Require Ctrl or not as per user preference
 		ELandscapeFoliageEditorControlType FoliageEditorControlType = GetDefault<ULevelEditorViewportSettings>()->FoliageEditorControlType;
@@ -2819,23 +2822,23 @@ bool FEdModeFoliage::InputKey(FEditorViewportClient* ViewportClient, FViewport*
 				bHandled = true;
 			}
 		}
+	}
 
-		if (!bHandled && bInSelectionMode)
+	if (!bHandled && (UISettings.GetLassoSelectToolSelected() || UISettings.GetSelectToolSelected()))
+	{
+		if (Event == IE_Pressed)
 		{
-			if (Event == IE_Pressed)
+			if (Key == EKeys::Platform_Delete)
 			{
-				if (Key == EKeys::Platform_Delete)
-				{
-					RemoveSelectedInstances(GetWorld());
+				RemoveSelectedInstances(GetWorld());
 
-					bHandled = true;
-				}
-				else if (Key == EKeys::End)
-				{
-					SnapSelectedInstancesToGround(GetWorld());
+				bHandled = true;
+			}
+			else if (Key == EKeys::End)
+			{
+				SnapSelectedInstancesToGround(GetWorld());
 
-					bHandled = true;
-				}
+				bHandled = true;
 			}
 		}
 	}
