Fix for non-unity issues introduced in CL# 2539880

diff --git a/Engine/Source/Editor/FoliageEdit/Private/FoliagePaletteItem.cpp b/Engine/Source/Editor/FoliageEdit/Private/FoliagePaletteItem.cpp
--- a/Engine/Source/Editor/FoliageEdit/Private/FoliagePaletteItem.cpp
+++ b/Engine/Source/Editor/FoliageEdit/Private/FoliagePaletteItem.cpp
@@ -11,6 +11,7 @@
 #include "Editor/UnrealEd/Public/AssetThumbnail.h"
 #include "Engine/StaticMesh.h"
 #include "AssetToolsModule.h"
+#include "FoliageType_InstancedStaticMesh.h"
 
 #define LOCTEXT_NAMESPACE "FoliageEd_Mode"
 
diff --git a/Engine/Source/Editor/FoliageEdit/Private/FoliageType_InstancedStaticMeshPaintingCustomization.cpp b/Engine/Source/Editor/FoliageEdit/Private/FoliageType_InstancedStaticMeshPaintingCustomization.cpp
--- a/Engine/Source/Editor/FoliageEdit/Private/FoliageType_InstancedStaticMeshPaintingCustomization.cpp
+++ b/Engine/Source/Editor/FoliageEdit/Private/FoliageType_InstancedStaticMeshPaintingCustomization.cpp
@@ -5,6 +5,7 @@
 #include "FoliageType_InstancedStaticMesh.h"
 #include "FoliageType_InstancedStaticMeshPaintingCustomization.h"
 #include "FoliageEdMode.h"
+#include "PropertyCustomizationHelpers.h"
 
 #define LOCTEXT_NAMESPACE "FoliageEd_Mode"
 
