Fix crash when foliage type is deleted and still inside the foliage editor.

UE-15447

diff --git a/Engine/Source/Editor/FoliageEdit/Private/FoliageEdMode.cpp b/Engine/Source/Editor/FoliageEdit/Private/FoliageEdMode.cpp
--- a/Engine/Source/Editor/FoliageEdit/Private/FoliageEdMode.cpp
+++ b/Engine/Source/Editor/FoliageEdit/Private/FoliageEdMode.cpp
@@ -299,6 +299,9 @@ void FEdModeFoliage::Enter()
 	FWorldDelegates::LevelAddedToWorld.AddSP(this, &FEdModeFoliage::NotifyLevelAddedToWorld);
 	FWorldDelegates::LevelRemovedFromWorld.AddSP(this, &FEdModeFoliage::NotifyLevelRemovedFromWorld);
 
+	FAssetRegistryModule& AssetRegistryModule = FModuleManager::GetModuleChecked<FAssetRegistryModule>(TEXT("AssetRegistry"));
+	AssetRegistryModule.Get().OnAssetRemoved().AddSP(this, &FEdModeFoliage::NotifyAssetRemoved);
+
 	// Force real-time viewports.  We'll back up the current viewport state so we can restore it when the
 	// user exits this mode.
 	const bool bWantRealTime = true;
@@ -348,6 +351,12 @@ void FEdModeFoliage::Exit()
 	FWorldDelegates::LevelAddedToWorld.RemoveAll(this);
 	FWorldDelegates::LevelRemovedFromWorld.RemoveAll(this);
 
+	if (FModuleManager::Get().IsModuleLoaded(TEXT("AssetRegistry")))
+	{
+		FAssetRegistryModule& AssetRegistryModule = FModuleManager::GetModuleChecked<FAssetRegistryModule>(TEXT("AssetRegistry"));
+		AssetRegistryModule.Get().OnAssetRemoved().RemoveAll(this);
+	}
+
 	GEditor->OnObjectsReplaced().RemoveAll(this);
 	
 	FoliageMeshList.Empty();
@@ -413,6 +422,14 @@ void FEdModeFoliage::NotifyLevelRemovedFromWorld(ULevel* InLevel, UWorld* InWorl
 	PopulateFoliageMeshList();
 }
 
+void FEdModeFoliage::NotifyAssetRemoved(const FAssetData& AssetInfo)
+{
+	if(UFoliageType* FoliageType = Cast<UFoliageType>(AssetInfo.GetAsset()))
+	{
+		PopulateFoliageMeshList();	//serialization already removes the asset from foliage actor. We just need to update the UI
+	}
+}
+
 /** When the user changes the current tool in the UI */
 void FEdModeFoliage::HandleToolChanged()
 {
@@ -2517,6 +2534,11 @@ bool FEdModeFoliage::CanPaint(const ULevel* InLevel)
 
 bool FEdModeFoliage::CanPaint(const UFoliageType* FoliageType, const ULevel* InLevel)
 {
+	if(FoliageType == nullptr)	//if asset has already been deleted we can't paint
+	{
+		return false;
+	}
+
 	// Non-shared objects can be painted only into their own level
 	// Assets can be painted everywhere
 	if (FoliageType->IsAsset() || FoliageType->GetOutermost() == InLevel->GetOutermost())
diff --git a/Engine/Source/Editor/FoliageEdit/Private/FoliageEdMode.h b/Engine/Source/Editor/FoliageEdit/Private/FoliageEdMode.h
--- a/Engine/Source/Editor/FoliageEdit/Private/FoliageEdMode.h
+++ b/Engine/Source/Editor/FoliageEdit/Private/FoliageEdMode.h
@@ -220,6 +220,9 @@ public:
 	void NotifyNewCurrentLevel();
 	void NotifyLevelAddedToWorld(ULevel* InLevel, UWorld* InWorld);
 	void NotifyLevelRemovedFromWorld(ULevel* InLevel, UWorld* InWorld);
+
+	/** Called when asset is removed */
+	void NotifyAssetRemoved(const FAssetData& AssetInfo);
 	
 	/**
 	 * Called when the mouse is moved over the viewport
