Fix crash when deleting the blueprint foliage type while it's in the foliage editor.

UE-15488

diff --git a/Engine/Source/Editor/FoliageEdit/Private/FoliageEdMode.cpp b/Engine/Source/Editor/FoliageEdit/Private/FoliageEdMode.cpp
--- a/Engine/Source/Editor/FoliageEdit/Private/FoliageEdMode.cpp
+++ b/Engine/Source/Editor/FoliageEdit/Private/FoliageEdMode.cpp
@@ -424,9 +424,14 @@ void FEdModeFoliage::NotifyLevelRemovedFromWorld(ULevel* InLevel, UWorld* InWorl
 
 void FEdModeFoliage::NotifyAssetRemoved(const FAssetData& AssetInfo)
 {
+	//TODO: This is not properly removing from the foliage actor. However, when we reload it will skip it.
+	//We need to properly fix this, but for now this prevents the crash
 	if(UFoliageType* FoliageType = Cast<UFoliageType>(AssetInfo.GetAsset()))
 	{
-		PopulateFoliageMeshList();	//serialization already removes the asset from foliage actor. We just need to update the UI
+		PopulateFoliageMeshList();	
+	}else if(UBlueprint* Blueprint = Cast<UBlueprint>(AssetInfo.GetAsset()))
+	{
+		PopulateFoliageMeshList();	
 	}
 }
 
