Fix foliage sticking around after deleting its asset. This fixes a crash when mapRebuild is triggered.

UE-15541

diff --git a/Engine/Source/Editor/FoliageEdit/Private/FoliageEditModule.cpp b/Engine/Source/Editor/FoliageEdit/Private/FoliageEditModule.cpp
--- a/Engine/Source/Editor/FoliageEdit/Private/FoliageEditModule.cpp
+++ b/Engine/Source/Editor/FoliageEdit/Private/FoliageEditModule.cpp
@@ -3,6 +3,7 @@
 
 #include "UnrealEd.h"
 #include "FoliageEditModule.h"
+#include "Runtime/AssetRegistry/Public/AssetRegistryModule.h"
 
 const FName FoliageEditAppIdentifier = FName(TEXT("FoliageEdApp"));
 
@@ -102,11 +103,24 @@ public:
 		}
 	}
 
+	void NotifyAssetRemoved(const FAssetData& AssetInfo)
+	{
+		// Go through all FoliageActors in the world and delete 
+		for(TObjectIterator<AInstancedFoliageActor> It; It; ++It)
+		{
+			AInstancedFoliageActor* IFA = *It;
+			IFA->CleanupDeletedFoliageType();
+		}
+	}
+
 	void SubscribeEvents()
 	{
 		GEngine->OnLevelActorDeleted().Remove(OnLevelActorDeletedDelegateHandle);
 		OnLevelActorDeletedDelegateHandle = GEngine->OnLevelActorDeleted().AddRaw(this, &FFoliageEditModule::OnLevelActorDeleted);
 
+		FAssetRegistryModule& AssetRegistryModule = FModuleManager::GetModuleChecked<FAssetRegistryModule>(TEXT("AssetRegistry"));
+		AssetRegistryModule.Get().OnAssetRemoved().AddRaw(this, &FFoliageEditModule::NotifyAssetRemoved);
+
 		auto ExperimentalSettings = GetMutableDefault<UEditorExperimentalSettings>();
 		ExperimentalSettings->OnSettingChanged().Remove(OnExperimentalSettingChangedDelegateHandle);
 		OnExperimentalSettingChangedDelegateHandle =  ExperimentalSettings->OnSettingChanged().AddRaw(this, &FFoliageEditModule::HandleExperimentalSettingChanged);
@@ -116,6 +130,12 @@ public:
 	{
 		GEngine->OnLevelActorDeleted().Remove(OnLevelActorDeletedDelegateHandle);
 		GetMutableDefault<UEditorExperimentalSettings>()->OnSettingChanged().Remove(OnExperimentalSettingChangedDelegateHandle);
+
+		if (FModuleManager::Get().IsModuleLoaded(TEXT("AssetRegistry")))
+		{
+			FAssetRegistryModule& AssetRegistryModule = FModuleManager::GetModuleChecked<FAssetRegistryModule>(TEXT("AssetRegistry"));
+			AssetRegistryModule.Get().OnAssetRemoved().RemoveAll(this);
+		}
 	}
 
 	void HandleExperimentalSettingChanged(FName PropertyName)
diff --git a/Engine/Source/Runtime/Foliage/Private/InstancedFoliage.cpp b/Engine/Source/Runtime/Foliage/Private/InstancedFoliage.cpp
--- a/Engine/Source/Runtime/Foliage/Private/InstancedFoliage.cpp
+++ b/Engine/Source/Runtime/Foliage/Private/InstancedFoliage.cpp
@@ -2393,8 +2393,34 @@ void AInstancedFoliageActor::OnPostApplyLevelOffset(ULevel* InLevel, UWorld* InW
 	}
 }
 
-#endif
 
+void AInstancedFoliageActor::CleanupDeletedFoliageType()
+{
+	for (auto& MeshPair : FoliageMeshes)
+	{
+		if(MeshPair.Key == nullptr)
+		{
+			FFoliageMeshInfo& MeshInfo = *MeshPair.Value;
+			TArray<int32> InstancesToRemove;
+			for (int32 InstanceIdx = 0; InstanceIdx < MeshInfo.Instances.Num(); InstanceIdx++)
+			{
+				InstancesToRemove.Add(InstanceIdx);
+			}
+
+			if (InstancesToRemove.Num())
+			{
+				MeshInfo.RemoveInstances(this, InstancesToRemove);
+			}
+		}
+		
+	}
+
+	while(FoliageMeshes.Remove(nullptr)){}	//remove entries from the map
+
+}
+
+
+#endif
 //
 // Serialize all our UObjects for RTGC 
 //
diff --git a/Engine/Source/Runtime/Foliage/Public/InstancedFoliageActor.h b/Engine/Source/Runtime/Foliage/Public/InstancedFoliageActor.h
--- a/Engine/Source/Runtime/Foliage/Public/InstancedFoliageActor.h
+++ b/Engine/Source/Runtime/Foliage/Public/InstancedFoliageActor.h
@@ -48,6 +48,7 @@ public:
 	virtual bool IsLevelBoundsRelevant() const override { return false; }
 #if WITH_EDITOR
 	virtual void Destroyed() override;
+	FOLIAGE_API void CleanupDeletedFoliageType();
 #endif
 	// End AActor interface.
 
