Foliage thumbnails now update when hovered over.

diff --git a/Engine/Source/Editor/FoliageEdit/Private/FoliagePaletteItem.cpp b/Engine/Source/Editor/FoliageEdit/Private/FoliagePaletteItem.cpp
--- a/Engine/Source/Editor/FoliageEdit/Private/FoliagePaletteItem.cpp
+++ b/Engine/Source/Editor/FoliageEdit/Private/FoliagePaletteItem.cpp
@@ -26,7 +26,7 @@ namespace EInstanceCountMagnitude
 	};
 }
 
-FFoliagePaletteItemModel::FFoliagePaletteItemModel(FFoliageMeshUIInfoPtr InTypeInfo, TSharedRef<SFoliagePalette> InFoliagePalette, FEdModeFoliage* InFoliageEditMode)
+FFoliagePaletteItemModel::FFoliagePaletteItemModel(FFoliageMeshUIInfoPtr InTypeInfo, TSharedRef<SFoliagePalette> InFoliagePalette, TSharedPtr<FAssetThumbnailPool> InThumbnailPool, FEdModeFoliage* InFoliageEditMode)
 	: TypeInfo(InTypeInfo)
 	, FoliagePalette(InFoliagePalette)
 	, FoliageEditMode(InFoliageEditMode)
@@ -44,10 +44,6 @@ FFoliagePaletteItemModel::FFoliagePaletteItemModel(FFoliageMeshUIInfoPtr InTypeI
 	{
 		DisplayFName = TypeInfo->Settings->GetStaticMesh()->GetFName();
 	}
-
-	// Create the thumbnail widget
-	FLevelEditorModule& LevelEditorModule = FModuleManager::LoadModuleChecked<FLevelEditorModule>("LevelEditor");
-	TSharedPtr<FAssetThumbnailPool> ThumbnailPool = LevelEditorModule.GetFirstLevelEditor()->GetThumbnailPool();
 	
 	FAssetData AssetData;
 	if (!IsBlueprint())
@@ -61,7 +57,7 @@ FFoliagePaletteItemModel::FFoliagePaletteItemModel(FFoliageMeshUIInfoPtr InTypeI
 	
 
 	int32 MaxThumbnailSize = FoliagePaletteConstants::ThumbnailSizeRange.Max;
-	TSharedPtr< FAssetThumbnail > Thumbnail = MakeShareable(new FAssetThumbnail(AssetData, MaxThumbnailSize, MaxThumbnailSize, ThumbnailPool));
+	TSharedPtr< FAssetThumbnail > Thumbnail = MakeShareable(new FAssetThumbnail(AssetData, MaxThumbnailSize, MaxThumbnailSize, InThumbnailPool));
 	
 	FAssetThumbnailConfig ThumbnailConfig;
 	FAssetToolsModule& AssetToolsModule = FModuleManager::LoadModuleChecked<FAssetToolsModule>(TEXT("AssetTools"));
diff --git a/Engine/Source/Editor/FoliageEdit/Private/FoliagePaletteItem.h b/Engine/Source/Editor/FoliageEdit/Private/FoliagePaletteItem.h
--- a/Engine/Source/Editor/FoliageEdit/Private/FoliagePaletteItem.h
+++ b/Engine/Source/Editor/FoliageEdit/Private/FoliagePaletteItem.h
@@ -23,7 +23,7 @@ namespace FoliagePaletteTreeColumns
 class FFoliagePaletteItemModel : public TSharedFromThis<FFoliagePaletteItemModel>
 {
 public:
-	FFoliagePaletteItemModel(FFoliageMeshUIInfoPtr InTypeInfo, TSharedRef<SFoliagePalette> InFoliagePalette, FEdModeFoliage* InFoliageEditMode);
+	FFoliagePaletteItemModel(FFoliageMeshUIInfoPtr InTypeInfo, TSharedRef<SFoliagePalette> InFoliagePalette, TSharedPtr<FAssetThumbnailPool> InThumbnailPool, FEdModeFoliage* InFoliageEditMode);
 
 	/** @return The foliage palette that contains the item */
 	TSharedPtr<SFoliagePalette> GetFoliagePalette() const;
@@ -86,10 +86,8 @@ private:
 
 	/** Gets the source asset type text */
 	FText GetSourceAssetTypeText() const;
-
 private:
 	TSharedPtr<SWidget> ThumbnailWidget;
-
 	FName DisplayFName;
 	FFoliageMeshUIInfoPtr TypeInfo;
 	TWeakPtr<SFoliagePalette> FoliagePalette;
diff --git a/Engine/Source/Editor/FoliageEdit/Private/SFoliagePalette.cpp b/Engine/Source/Editor/FoliageEdit/Private/SFoliagePalette.cpp
--- a/Engine/Source/Editor/FoliageEdit/Private/SFoliagePalette.cpp
+++ b/Engine/Source/Editor/FoliageEdit/Private/SFoliagePalette.cpp
@@ -160,6 +160,8 @@ void SFoliagePalette::Construct(const FArguments& InArgs)
 	UICommandList = MakeShareable(new FUICommandList);
 	BindCommands();
 
+	ThumbnailPool = MakeShareable(new FAssetThumbnailPool(25, TAttribute<bool>::Create(TAttribute<bool>::FGetter::CreateSP(this, &SFoliagePalette::IsHovered))));
+
 	TypeFilter = MakeShareable(new FoliageTypeTextFilter(
 		FoliageTypeTextFilter::FItemToStringArray::CreateSP(this, &SFoliagePalette::GetPaletteItemFilterString)));
 
@@ -435,7 +437,7 @@ void SFoliagePalette::UpdateThumbnailForType(UFoliageType* FoliageType)
 		{
 			const bool bItemIsSelected = GetActiveViewWidget()->IsItemSelected(Item);
 
-			Item = MakeShareable(new FFoliagePaletteItemModel(Item->GetTypeUIInfo(), SharedThis(this), FoliageEditMode));
+			Item = MakeShareable(new FFoliagePaletteItemModel(Item->GetTypeUIInfo(), SharedThis(this), ThumbnailPool, FoliageEditMode));
 			if (bItemIsSelected)
 			{
 				GetActiveViewWidget()->SetItemSelection(Item, true);
@@ -1346,7 +1348,7 @@ EActiveTimerReturnType SFoliagePalette::UpdatePaletteItems(double InCurrentTime,
 		PaletteItems.Empty(AllTypesList.Num());
 		for (const FFoliageMeshUIInfoPtr& TypeInfo : AllTypesList)
 		{
-			PaletteItems.Add(MakeShareable(new FFoliagePaletteItemModel(TypeInfo, SharedThis(this), FoliageEditMode)));
+			PaletteItems.Add(MakeShareable(new FFoliagePaletteItemModel(TypeInfo, SharedThis(this), ThumbnailPool, FoliageEditMode)));
 		}
 
 		// Restore the selection
diff --git a/Engine/Source/Editor/FoliageEdit/Private/SFoliagePalette.h b/Engine/Source/Editor/FoliageEdit/Private/SFoliagePalette.h
--- a/Engine/Source/Editor/FoliageEdit/Private/SFoliagePalette.h
+++ b/Engine/Source/Editor/FoliageEdit/Private/SFoliagePalette.h
@@ -260,6 +260,9 @@ private:
 	/** Command list for binding functions for the context menu. */
 	TSharedPtr<FUICommandList> UICommandList;
 
+	/** Thumbnail pool for rendering mesh thumbnails */
+	TSharedPtr<FAssetThumbnailPool> ThumbnailPool;
+
 	FEdModeFoliage* FoliageEditMode;
 
 	bool bItemsNeedRebuild : 1;
