Fix non-unity errors

diff --git a/Engine/Source/Editor/FoliageEdit/Private/FoliagePaletteItem.h b/Engine/Source/Editor/FoliageEdit/Private/FoliagePaletteItem.h
--- a/Engine/Source/Editor/FoliageEdit/Private/FoliagePaletteItem.h
+++ b/Engine/Source/Editor/FoliageEdit/Private/FoliagePaletteItem.h
@@ -23,7 +23,7 @@ namespace FoliagePaletteTreeColumns
 class FFoliagePaletteItemModel : public TSharedFromThis<FFoliagePaletteItemModel>
 {
 public:
-	FFoliagePaletteItemModel(FFoliageMeshUIInfoPtr InTypeInfo, TSharedRef<SFoliagePalette> InFoliagePalette, TSharedPtr<FAssetThumbnailPool> InThumbnailPool, FEdModeFoliage* InFoliageEditMode);
+	FFoliagePaletteItemModel(FFoliageMeshUIInfoPtr InTypeInfo, TSharedRef<SFoliagePalette> InFoliagePalette, TSharedPtr<class FAssetThumbnailPool> InThumbnailPool, FEdModeFoliage* InFoliageEditMode);
 
 	/** @return The foliage palette that contains the item */
 	TSharedPtr<SFoliagePalette> GetFoliagePalette() const;
diff --git a/Engine/Source/Editor/FoliageEdit/Private/SFoliagePalette.h b/Engine/Source/Editor/FoliageEdit/Private/SFoliagePalette.h
--- a/Engine/Source/Editor/FoliageEdit/Private/SFoliagePalette.h
+++ b/Engine/Source/Editor/FoliageEdit/Private/SFoliagePalette.h
@@ -261,7 +261,7 @@ private:
 	TSharedPtr<FUICommandList> UICommandList;
 
 	/** Thumbnail pool for rendering mesh thumbnails */
-	TSharedPtr<FAssetThumbnailPool> ThumbnailPool;
+	TSharedPtr<class FAssetThumbnailPool> ThumbnailPool;
 
 	FEdModeFoliage* FoliageEditMode;
 
diff --git a/Engine/Source/Runtime/Engine/Private/Particles/ParticleModules_Parameter.cpp b/Engine/Source/Runtime/Engine/Private/Particles/ParticleModules_Parameter.cpp
--- a/Engine/Source/Runtime/Engine/Private/Particles/ParticleModules_Parameter.cpp
+++ b/Engine/Source/Runtime/Engine/Private/Particles/ParticleModules_Parameter.cpp
@@ -16,6 +16,7 @@
 #include "Particles/ParticleLODLevel.h"
 #include "Particles/ParticleModuleRequired.h"
 #include "Engine/InterpCurveEdSetup.h"
+#include "Particles/Material/ParticleModuleMeshMaterial.h"
 
 UParticleModuleParameterBase::UParticleModuleParameterBase(const FObjectInitializer& ObjectInitializer)
 	: Super(ObjectInitializer)
