Foliage painting always trying to paint through filtered geometry
For example: when static meshes are filtered foliage still can be painted on landscape bellow static mesh
UE-7012

diff --git a/Engine/Source/Editor/FoliageEdit/Private/FoliageEdMode.cpp b/Engine/Source/Editor/FoliageEdit/Private/FoliageEdMode.cpp
--- a/Engine/Source/Editor/FoliageEdit/Private/FoliageEdMode.cpp
+++ b/Engine/Source/Editor/FoliageEdit/Private/FoliageEdMode.cpp
@@ -159,6 +159,41 @@ public:
 	}
 };
 
+//
+// Painting filtering options
+//
+struct FFoliagePaintingGeometryFilter
+{
+	bool bAllowLandscape;
+	bool bAllowStaticMesh;
+	bool bAllowBSP;
+	bool bAllowTranslucent;
+
+	FFoliagePaintingGeometryFilter(const FFoliageUISettings& InUISettings)
+		: bAllowLandscape(InUISettings.bFilterLandscape)
+		, bAllowStaticMesh(InUISettings.bFilterStaticMesh)
+		, bAllowBSP(InUISettings.bFilterBSP)
+		, bAllowTranslucent(InUISettings.bFilterTranslucent)
+	{
+	}
+
+	bool operator() (const UPrimitiveComponent* Component) const
+	{
+		if (Component)
+		{
+			bool bShallNotPass = 
+				(!bAllowLandscape	&& Component->IsA(ULandscapeHeightfieldCollisionComponent::StaticClass())) ||
+				(!bAllowStaticMesh	&& Component->IsA(UStaticMeshComponent::StaticClass())) ||
+				(!bAllowBSP			&& Component->IsA(UModelComponent::StaticClass())) ||
+				(!bAllowTranslucent	&& Component->GetMaterial(0) && IsTranslucentBlendMode(Component->GetMaterial(0)->GetBlendMode()));
+			
+			return !bShallNotPass;
+		}
+		
+		return false;		
+	}
+};
+
 
 //
 // FEdModeFoliage
@@ -601,18 +636,11 @@ void FEdModeFoliage::FoliageBrushTrace(FEditorViewportClient* ViewportClient, in
 			FHitResult Hit;
 			UWorld* World = ViewportClient->GetWorld();
 			static FName NAME_FoliageBrush = FName(TEXT("FoliageBrush"));
-			if (AInstancedFoliageActor::FoliageTrace(World, Hit, FDesiredFoliageInstance(Start, End), NAME_FoliageBrush))
+			FFoliagePaintingGeometryFilter FilterFunc = FFoliagePaintingGeometryFilter(UISettings);
+			if (AInstancedFoliageActor::FoliageTrace(World, Hit, FDesiredFoliageInstance(Start, End), NAME_FoliageBrush, false, FilterFunc))
 			{
-				// Check filters
 				UPrimitiveComponent* PrimComp = Hit.Component.Get();
-				UMaterialInterface* Material = PrimComp ? PrimComp->GetMaterial(0) : nullptr;
-
-				if (PrimComp &&
-					(UISettings.bFilterLandscape || !PrimComp->IsA(ULandscapeHeightfieldCollisionComponent::StaticClass())) &&
-					(UISettings.bFilterStaticMesh || !PrimComp->IsA(UStaticMeshComponent::StaticClass())) &&
-					(UISettings.bFilterBSP || !PrimComp->IsA(UModelComponent::StaticClass())) &&
-					(UISettings.bFilterTranslucent || !Material || !IsTranslucentBlendMode(Material->GetBlendMode())) &&
-					(CanPaint(PrimComp->GetComponentLevel())))
+				if (CanPaint(PrimComp->GetComponentLevel()))
 				{
 					if (!bAdjustBrushRadius)
 					{
@@ -842,31 +870,6 @@ bool FilterByWeight(float Weight, const UFoliageType* Settings)
 	return Weight < FMath::Max(SMALL_NUMBER, WeightNeeded);
 }
 
-bool GeometryFilterCheck(const FHitResult& Hit, const UWorld* InWorld, const FDesiredFoliageInstance& DesiredInst, const FFoliageUISettings* UISettings)
-{
-	if (UPrimitiveComponent* PrimComp = Hit.Component.Get())
-	{
-		if (DesiredInst.PlacementMode == EFoliagePlacementMode::Manual)
-		{
-			UMaterialInterface* Material = PrimComp ? PrimComp->GetMaterial(0) : nullptr;
-			if ((!UISettings->bFilterLandscape && PrimComp->IsA(ULandscapeHeightfieldCollisionComponent::StaticClass())) ||
-				(!UISettings->bFilterStaticMesh && PrimComp->IsA(UStaticMeshComponent::StaticClass())) ||
-				(!UISettings->bFilterBSP && PrimComp->IsA(UModelComponent::StaticClass())) ||
-				(!UISettings->bFilterTranslucent && Material && IsTranslucentBlendMode(Material->GetBlendMode()))
-				)
-			{
-				return false;
-			}
-		}
-	}
-	else
-	{
-		return false;
-	}
-
-	return true;
-}
-
 bool FEdModeFoliage::VertexMaskCheck(const FHitResult& Hit, const UFoliageType* Settings)
 {
 	if (Settings->VertexColorMask != FOLIAGEVERTEXCOLORMASK_Disabled && Hit.FaceIndex != INDEX_NONE)
@@ -913,16 +916,25 @@ void FEdModeFoliage::CalculatePotentialInstances_ThreadSafe(const UWorld* InWorl
 		Bucket.Reserve(DesiredInstances->Num());
 	}
 
+	// Filter by geometry type when painting folige manualy
+	FFoliageTraceFilterFunc TraceFilterFunc_Manual = FFoliagePaintingGeometryFilter(*UISettings);
+
 	for (int32 InstanceIdx = StartIdx; InstanceIdx <= LastIdx; ++InstanceIdx)
 	{
 		const FDesiredFoliageInstance& DesiredInst = (*DesiredInstances)[InstanceIdx];
 		FHitResult Hit;
 		static FName NAME_AddFoliageInstances = FName(TEXT("AddFoliageInstances"));
-		if (AInstancedFoliageActor::FoliageTrace(InWorld, Hit, DesiredInst, NAME_AddFoliageInstances, true))
+		
+		FFoliageTraceFilterFunc TraceFilterFunc;
+		if (DesiredInst.PlacementMode == EFoliagePlacementMode::Manual)
+		{
+			TraceFilterFunc = TraceFilterFunc_Manual;
+		}
+		
+		if (AInstancedFoliageActor::FoliageTrace(InWorld, Hit, DesiredInst, NAME_AddFoliageInstances, true, TraceFilterFunc))
 		{
 			float HitWeight = 1.f;
-			const bool bValidInstance = GeometryFilterCheck(Hit, InWorld, DesiredInst, UISettings)
-										&& CheckLocationForPotentialInstance_ThreadSafe(Settings, Hit.ImpactPoint, Hit.ImpactNormal)
+			const bool bValidInstance = CheckLocationForPotentialInstance_ThreadSafe(Settings, Hit.ImpactPoint, Hit.ImpactNormal)
 										&& VertexMaskCheck(Hit, Settings)
 										&& LandscapeLayerCheck(Hit, Settings, LocalCache, HitWeight);
 
@@ -952,11 +964,20 @@ void FEdModeFoliage::CalculatePotentialInstances(const UWorld* InWorld, const UF
 		Bucket.Reserve(DesiredInstances.Num());
 	}
 
+	// Filter by geometry type when painting foliage manually
+	FFoliageTraceFilterFunc TraceFilterFunc_Manual = FFoliagePaintingGeometryFilter(*UISettings);
+
 	for (const FDesiredFoliageInstance& DesiredInst : DesiredInstances)
 	{
+		FFoliageTraceFilterFunc TraceFilterFunc;
+		if (DesiredInst.PlacementMode == EFoliagePlacementMode::Manual)
+		{
+			TraceFilterFunc = TraceFilterFunc_Manual;
+		}
+				
 		FHitResult Hit;
 		static FName NAME_AddFoliageInstances = FName(TEXT("AddFoliageInstances"));
-		if (AInstancedFoliageActor::FoliageTrace(InWorld, Hit, DesiredInst, NAME_AddFoliageInstances, true))
+		if (AInstancedFoliageActor::FoliageTrace(InWorld, Hit, DesiredInst, NAME_AddFoliageInstances, true, TraceFilterFunc))
 		{
 			float HitWeight = 1.f;
 
@@ -970,11 +991,6 @@ void FEdModeFoliage::CalculatePotentialInstances(const UWorld* InWorld, const UF
 				continue;
 			}
 
-			if (!GeometryFilterCheck(Hit, InWorld, DesiredInst, UISettings))
-			{
-				continue;
-			}
-
 			const bool bValidInstance =	CheckLocationForPotentialInstance(InWorld, Settings, Hit.ImpactPoint, Hit.ImpactNormal, PotentialInstanceLocations, PotentialInstanceHash)
 										&& VertexMaskCheck(Hit, Settings)
 										&& LandscapeLayerCheck(Hit, Settings, LocalCache, HitWeight);
@@ -1165,21 +1181,17 @@ void FEdModeFoliage::RemoveInstancesForBrush(UWorld* InWorld, const UFoliageType
 
 		if (!UISettings.bFilterLandscape || !UISettings.bFilterStaticMesh || !UISettings.bFilterBSP || !UISettings.bFilterTranslucent)
 		{
+			FFoliagePaintingGeometryFilter GeometryFilterFunc(UISettings);
+			
 			// Filter PotentialInstancesToRemove
 			for (int32 Idx = 0; Idx < PotentialInstancesToRemove.Num(); Idx++)
 			{
 				auto BaseId = MeshInfo->Instances[PotentialInstancesToRemove[Idx]].BaseId;
 				auto BasePtr = IFA->InstanceBaseCache.GetInstanceBasePtr(BaseId);
 				UPrimitiveComponent* Base = Cast<UPrimitiveComponent>(BasePtr.Get());
-				UMaterialInterface* Material = Base ? Base->GetMaterial(0) : nullptr;
 
 				// Check if instance is candidate for removal based on filter settings
-				if (Base && (
-					(!UISettings.bFilterLandscape && Base->IsA(ULandscapeHeightfieldCollisionComponent::StaticClass())) ||
-					(!UISettings.bFilterStaticMesh && Base->IsA(UStaticMeshComponent::StaticClass())) ||
-					(!UISettings.bFilterBSP && Base->IsA(UModelComponent::StaticClass())) ||
-					(!UISettings.bFilterTranslucent && Material && IsTranslucentBlendMode(Material->GetBlendMode()))
-					))
+				if (Base && !GeometryFilterFunc(Base))
 				{
 					// Instance should not be removed, so remove it from the removal list.
 					PotentialInstancesToRemove.RemoveAtSwap(Idx);
diff --git a/Engine/Source/Runtime/Foliage/Private/InstancedFoliage.cpp b/Engine/Source/Runtime/Foliage/Private/InstancedFoliage.cpp
--- a/Engine/Source/Runtime/Foliage/Private/InstancedFoliage.cpp
+++ b/Engine/Source/Runtime/Foliage/Private/InstancedFoliage.cpp
@@ -2442,7 +2442,7 @@ void AInstancedFoliageActor::AddReferencedObjects(UObject* InThis, FReferenceCol
 	Super::AddReferencedObjects(This, Collector);
 }
 
-bool AInstancedFoliageActor::FoliageTrace(const UWorld* InWorld, FHitResult& OutHit, const FDesiredFoliageInstance& DesiredInstance, FName InTraceTag, bool InbReturnFaceIndex)
+bool AInstancedFoliageActor::FoliageTrace(const UWorld* InWorld, FHitResult& OutHit, const FDesiredFoliageInstance& DesiredInstance, FName InTraceTag, bool InbReturnFaceIndex, FFoliageTraceFilterFunc FilterFunc)
 {
 	FCollisionQueryParams QueryParams(InTraceTag, true);
 	QueryParams.bReturnFaceIndex = InbReturnFaceIndex;
@@ -2494,8 +2494,16 @@ bool AInstancedFoliageActor::FoliageTrace(const UWorld* InWorld, FHitResult& Out
 			continue;
 		}
 
-		if (Hit.Component.IsValid() && Hit.Component->GetComponentLevel())
+		const UPrimitiveComponent* HitComponent = Hit.Component.Get();
+
+		if (HitComponent && HitComponent->GetComponentLevel())
 		{
+			if (FilterFunc && FilterFunc(HitComponent) == false)
+			{
+				// supplied filter does not like this component, so keep iterating
+				continue;
+			}
+			
 			OutHit = Hit;
 			return (DesiredInstance.PlacementMode != EFoliagePlacementMode::Procedural) || bInsideProceduralVolume;
 		}
diff --git a/Engine/Source/Runtime/Foliage/Public/InstancedFoliageActor.h b/Engine/Source/Runtime/Foliage/Public/InstancedFoliageActor.h
--- a/Engine/Source/Runtime/Foliage/Public/InstancedFoliageActor.h
+++ b/Engine/Source/Runtime/Foliage/Public/InstancedFoliageActor.h
@@ -20,6 +20,9 @@ struct FFoliageInstance;
 struct FHitResult;
 struct FDesiredFoliageInstance;
 
+// Function for filtering out hit components during FoliageTrace
+typedef TFunction<bool(const UPrimitiveComponent*)> FFoliageTraceFilterFunc;
+
 UCLASS(notplaceable, hidecategories = Object, MinimalAPI, NotBlueprintable)
 class AInstancedFoliageActor : public AActor
 {
@@ -94,7 +97,7 @@ public:
 	*/
 	static FOLIAGE_API AInstancedFoliageActor* GetInstancedFoliageActorForLevel(ULevel* Level, bool bCreateIfNone = false);
 
-	static FOLIAGE_API bool FoliageTrace(const UWorld* InWorld, FHitResult& OutHit, const FDesiredFoliageInstance& DesiredInstance, FName InTraceTag = NAME_None, bool InbReturnFaceIndex = false);
+	static FOLIAGE_API bool FoliageTrace(const UWorld* InWorld, FHitResult& OutHit, const FDesiredFoliageInstance& DesiredInstance, FName InTraceTag = NAME_None, bool InbReturnFaceIndex = false, FFoliageTraceFilterFunc FilterFunc = FFoliageTraceFilterFunc());
 	static FOLIAGE_API bool CheckCollisionWithWorld(const UWorld* InWorld, const UFoliageType* Settings, const FFoliageInstance& Inst, const FVector& HitNormal, const FVector& HitLocation);
 
 #if WITH_EDITOR
