FoliageTool: Mass toggle filtered items instead of selected

diff --git a/Engine/Source/Editor/FoliageEdit/Private/SFoliagePalette.cpp b/Engine/Source/Editor/FoliageEdit/Private/SFoliagePalette.cpp
--- a/Engine/Source/Editor/FoliageEdit/Private/SFoliagePalette.cpp
+++ b/Engine/Source/Editor/FoliageEdit/Private/SFoliagePalette.cpp
@@ -1191,7 +1191,7 @@ ECheckBoxState SFoliagePalette::GetState_AllMeshes() const
 	bool bHasChecked = false;
 	bool bHasUnchecked = false;
 
-	for (FFoliagePaletteItemModelPtr& PaletteItem : GetActiveViewWidget()->GetSelectedItems())
+	for (const FFoliagePaletteItemModelPtr& PaletteItem : FilteredItems)
 	{
 		if (PaletteItem->IsActive())
 		{
@@ -1214,7 +1214,7 @@ ECheckBoxState SFoliagePalette::GetState_AllMeshes() const
 void SFoliagePalette::OnCheckStateChanged_AllMeshes(ECheckBoxState InState)
 {
 	const bool bActivate = InState == ECheckBoxState::Checked;
-	for (FFoliagePaletteItemModelPtr& PaletteItem : GetActiveViewWidget()->GetSelectedItems())
+	for (FFoliagePaletteItemModelPtr& PaletteItem : FilteredItems)
 	{
 		PaletteItem->SetTypeActiveInPalette(bActivate);
 	}
