Fixed: Crash resimulating a foliage volume UE-18304

diff --git a/Engine/Source/Editor/FoliageEdit/Private/FoliageEdMode.cpp b/Engine/Source/Editor/FoliageEdit/Private/FoliageEdMode.cpp
--- a/Engine/Source/Editor/FoliageEdit/Private/FoliageEdMode.cpp
+++ b/Engine/Source/Editor/FoliageEdit/Private/FoliageEdMode.cpp
@@ -916,9 +916,6 @@ void FEdModeFoliage::CalculatePotentialInstances_ThreadSafe(const UWorld* InWorl
 		Bucket.Reserve(DesiredInstances->Num());
 	}
 
-	// Filter by geometry type when painting folige manualy
-	FFoliageTraceFilterFunc TraceFilterFunc_Manual = FFoliagePaintingGeometryFilter(*UISettings);
-
 	for (int32 InstanceIdx = StartIdx; InstanceIdx <= LastIdx; ++InstanceIdx)
 	{
 		const FDesiredFoliageInstance& DesiredInst = (*DesiredInstances)[InstanceIdx];
@@ -926,9 +923,10 @@ void FEdModeFoliage::CalculatePotentialInstances_ThreadSafe(const UWorld* InWorl
 		static FName NAME_AddFoliageInstances = FName(TEXT("AddFoliageInstances"));
 		
 		FFoliageTraceFilterFunc TraceFilterFunc;
-		if (DesiredInst.PlacementMode == EFoliagePlacementMode::Manual)
+		if (DesiredInst.PlacementMode == EFoliagePlacementMode::Manual && UISettings != nullptr)
 		{
-			TraceFilterFunc = TraceFilterFunc_Manual;
+			// Enable geometry filters when painting foliage manually
+			TraceFilterFunc = FFoliagePaintingGeometryFilter(*UISettings);
 		}
 		
 		if (AInstancedFoliageActor::FoliageTrace(InWorld, Hit, DesiredInst, NAME_AddFoliageInstances, true, TraceFilterFunc))
@@ -964,15 +962,13 @@ void FEdModeFoliage::CalculatePotentialInstances(const UWorld* InWorld, const UF
 		Bucket.Reserve(DesiredInstances.Num());
 	}
 
-	// Filter by geometry type when painting foliage manually
-	FFoliageTraceFilterFunc TraceFilterFunc_Manual = FFoliagePaintingGeometryFilter(*UISettings);
-
 	for (const FDesiredFoliageInstance& DesiredInst : DesiredInstances)
 	{
 		FFoliageTraceFilterFunc TraceFilterFunc;
-		if (DesiredInst.PlacementMode == EFoliagePlacementMode::Manual)
+		if (DesiredInst.PlacementMode == EFoliagePlacementMode::Manual && UISettings != nullptr)
 		{
-			TraceFilterFunc = TraceFilterFunc_Manual;
+			// Enable geometry filters when painting foliage manually
+			TraceFilterFunc = FFoliagePaintingGeometryFilter(*UISettings);
 		}
 				
 		FHitResult Hit;
