Fixed: Foliage does not spawn on Hierarchical Instanced Static Mesh Components UE-18175

diff --git a/Engine/Source/Editor/FoliageEdit/Private/FoliageEdMode.cpp b/Engine/Source/Editor/FoliageEdit/Private/FoliageEdMode.cpp
--- a/Engine/Source/Editor/FoliageEdit/Private/FoliageEdMode.cpp
+++ b/Engine/Source/Editor/FoliageEdit/Private/FoliageEdMode.cpp
@@ -706,6 +706,13 @@ void FEdModeFoliage::GetRandomVectorInBrush(FVector& OutStart, FVector& OutEnd)
 	OutEnd		= BrushLocation + UISettings.GetRadius() * (Point + Rw);
 }
 
+static bool IsWithinSlopeAngle(float NormalZ, float MinAngle, float MaxAngle, float Tolerance = SMALL_NUMBER)
+{
+	const float MaxNormalAngle = FMath::Cos(FMath::DegreesToRadians(MaxAngle));
+	const float MinNormalAngle = FMath::Cos(FMath::DegreesToRadians(MinAngle));
+	return !(MaxNormalAngle > (NormalZ+Tolerance) || MinNormalAngle < (NormalZ-Tolerance));
+}
+
 /** This does not check for overlaps or density */
 static bool CheckLocationForPotentialInstance_ThreadSafe(const UFoliageType* Settings, const FVector& Location, const FVector& Normal)
 {
@@ -716,14 +723,8 @@ static bool CheckLocationForPotentialInstance_ThreadSafe(const UFoliageType* Set
 	}
 
 	// Check slope
-	const float MaxNormalAngle = FMath::Cos(FMath::DegreesToRadians(Settings->GroundSlopeAngle.Max));
-	const float MinNormalAngle = FMath::Cos(FMath::DegreesToRadians(Settings->GroundSlopeAngle.Min));
-	if (MaxNormalAngle > Normal.Z || MinNormalAngle < Normal.Z)	//keep in mind Hit.ImpactNormal.Z is (0,0,1) dot normal. However, ground slope is with relation to plane vector, not plane normal - so we swap comparisons
-	{
-		return false;
-	}
-
-	return true;
+	// ImpactNormal sometimes is slightly non-normalized, so compare slope with some little deviation
+	return IsWithinSlopeAngle(Normal.Z, Settings->GroundSlopeAngle.Min, Settings->GroundSlopeAngle.Max, SMALL_NUMBER);
 }
 
 static bool CheckForOverlappingSphere(AInstancedFoliageActor* IFA, const UFoliageType* Settings, const FSphere& Sphere)
@@ -1685,9 +1686,7 @@ void FEdModeFoliage::ReapplyInstancesForBrush(UWorld* InWorld, AInstancedFoliage
 					// Cull instances that don't meet the ground slope check.
 					if (Settings->ReapplyGroundSlope)
 					{
-						const float MaxNormalAngle = FMath::Cos(FMath::DegreesToRadians(Settings->GroundSlopeAngle.Max));
-						const float MinNormalAngle = FMath::Cos(FMath::DegreesToRadians(Settings->GroundSlopeAngle.Min));
-						if (MaxNormalAngle > Hit.Normal.Z || MinNormalAngle < Hit.Normal.Z)
+						if (!IsWithinSlopeAngle(Hit.Normal.Z, Settings->GroundSlopeAngle.Min, Settings->GroundSlopeAngle.Max))
 						{
 							InstancesToDelete.Add(InstanceIndex);
 							if (bReapplyLocation)
@@ -2137,9 +2136,7 @@ void FEdModeFoliage::ApplyPaintBucket_Add(AActor* Actor)
 				FFoliagePaintBucketTriangle& Triangle = PotentialTriangles[TriIdx];
 
 				// Check if we can reject this triangle based on normal.
-				const float MaxNormalAngle = FMath::Cos(FMath::DegreesToRadians(Settings->GroundSlopeAngle.Max));
-				const float MinNormalAngle = FMath::Cos(FMath::DegreesToRadians(Settings->GroundSlopeAngle.Min));
-				if (Triangle.WorldNormal.Z < MaxNormalAngle || Triangle.WorldNormal.Z > MinNormalAngle)
+				if (!IsWithinSlopeAngle(Triangle.WorldNormal.Z, Settings->GroundSlopeAngle.Min, Settings->GroundSlopeAngle.Max))
 				{
 					continue;
 				}
