Fixed crash when changing foliage mesh painted across multiple levels.

diff --git a/Engine/Source/Editor/FoliageEdit/Private/SFoliagePalette.cpp b/Engine/Source/Editor/FoliageEdit/Private/SFoliagePalette.cpp
--- a/Engine/Source/Editor/FoliageEdit/Private/SFoliagePalette.cpp
+++ b/Engine/Source/Editor/FoliageEdit/Private/SFoliagePalette.cpp
@@ -722,7 +722,12 @@ FText SFoliagePalette::GetSearchText() const
 
 void SFoliagePalette::OnSelectionChanged(FFoliagePaletteItemModelPtr Item, ESelectInfo::Type SelectInfo)
 {
-	RefreshDetailsWidget();
+	if (SelectInfo != ESelectInfo::Direct)
+	{
+		// If we set the selection in code it's because we're restoring a previous selection after changing the palette.
+		// Refreshing the details palette causes problems if we're in the middle of a details change which caused a thumbnail invalidation.
+		RefreshDetailsWidget();
+	}
 
 	bIsUneditableFoliageTypeSelected = false;
 	for (FFoliagePaletteItemModelPtr& PaletteItem : GetActiveViewWidget()->GetSelectedItems())
