[NO NON-PII COMMIT MESSAGE]

diff --git a/Engine/Source/Editor/FoliageEdit/FoliageEdit.build.cs b/Engine/Source/Editor/FoliageEdit/FoliageEdit.build.cs
--- a/Engine/Source/Editor/FoliageEdit/FoliageEdit.build.cs
+++ b/Engine/Source/Editor/FoliageEdit/FoliageEdit.build.cs
@@ -7,7 +7,8 @@ public class FoliageEdit : ModuleRules
 	public FoliageEdit(TargetInfo Target)
 	{
 		PublicDependencyModuleNames.AddRange(
-			new string[] {
+			new string[] 
+			{
 				"Core",
 				"CoreUObject",
 				"InputCore",
@@ -25,5 +26,13 @@ public class FoliageEdit : ModuleRules
                 "Foliage",
 			}
 		);
-	}
+
+		PrivateDependencyModuleNames.AddRange(
+			new string[]
+			{
+				"ViewportInteraction",
+				"VREditor"
+			}
+		);
+    }
 }
diff --git a/Engine/Source/Editor/FoliageEdit/Private/FoliageEdMode.cpp b/Engine/Source/Editor/FoliageEdit/Private/FoliageEdMode.cpp
--- a/Engine/Source/Editor/FoliageEdit/Private/FoliageEdMode.cpp
+++ b/Engine/Source/Editor/FoliageEdit/Private/FoliageEdMode.cpp
@@ -36,6 +36,11 @@
 #include "FoliageType_InstancedStaticMesh.h"
 #include "Components/BrushComponent.h"
 
+// VR Editor
+#include "IVREditorModule.h"
+#include "ViewportWorldInteraction.h"
+#include "VREditorInteractor.h"
+
 
 #define LOCTEXT_NAMESPACE "FoliageEdMode"
 #define FOLIAGE_SNAP_TRACE (10000.f)
@@ -170,6 +175,8 @@ FEdModeFoliage::FEdModeFoliage()
 	, bCanAltDrag(false)
 	, bAdjustBrushRadius(false)
 	, FoliageMeshListSortMode(EColumnSortMode::Ascending)
+	, bIsPainting(false)
+	, FoliageInteractor(nullptr)
 {
 	// Load resources and construct brush component
 	UMaterial* BrushMaterial = nullptr;
@@ -337,15 +344,41 @@ void FEdModeFoliage::Enter()
 		
 	// Update UI
 	NotifyNewCurrentLevel();
+
+	// Register to find out when foliage mode is activated
+	GetModeManager()->OnEditorModeChanged().AddRaw(this, &FEdModeFoliage::OnEditorModeChanged);
+
+	// Register for VR input events
+	IVREditorMode* VREditorMode = static_cast<IVREditorMode*>(GetModeManager()->GetActiveMode(IVREditorModule::Get().GetVREditorModeID()));
+	if (VREditorMode != nullptr)
+	{
+		VREditorMode->GetWorldInteraction().OnViewportInteractionInputAction().RemoveAll(this);
+		VREditorMode->GetWorldInteraction().OnViewportInteractionInputAction().AddRaw(this, &FEdModeFoliage::OnVRAction);
+		VREditorMode->GetWorldInteraction().OnViewportInteractionHoverUpdate().AddRaw(this, &FEdModeFoliage::OnVRHoverUpdate);
+	}
 }
 
 /** FEdMode: Called when the mode is exited */
 void FEdModeFoliage::Exit()
 {
+	// Unregister VR mode from event handlers
+	if (IVREditorModule::IsAvailable())
+	{
+		IVREditorMode* VREditorMode = static_cast<IVREditorMode*>(GetModeManager()->GetActiveMode(IVREditorModule::Get().GetVREditorModeID()));
+		if (VREditorMode != nullptr)
+		{
+			VREditorMode->GetWorldInteraction().OnViewportInteractionInputAction().RemoveAll(this);
+			FoliageInteractor = nullptr;
+		}
+	}
+
+	// Unregister foliage mode from event handlers
+	GetModeManager()->OnEditorModeChanged().RemoveAll(this);
+
 	FToolkitManager::Get().CloseToolkit(Toolkit.ToSharedRef());
 	Toolkit.Reset();
 
-	//
+	// Remove delegates
 	FEditorDelegates::NewCurrentLevel.RemoveAll(this);
 	FWorldDelegates::LevelAddedToWorld.RemoveAll(this);
 	FWorldDelegates::LevelRemovedFromWorld.RemoveAll(this);
@@ -398,6 +431,136 @@ void FEdModeFoliage::Exit()
 	FEdMode::Exit();
 }
 
+void FEdModeFoliage::OnEditorModeChanged(FEdMode* EditorMode, bool bEntered)
+{
+	// Detect if we're switching to VR Editor mode
+	// VR Editor mode may have gone away, so re-register for events in case it's a different object than the one we were originally bound to
+	check(EditorMode != nullptr);
+	FoliageInteractor = nullptr;
+	if (bEntered && EditorMode->GetID() == IVREditorModule::Get().GetVREditorModeID())
+	{
+		IVREditorMode* VREditorMode = static_cast<IVREditorMode*>(EditorMode);
+		if (VREditorMode != nullptr)
+		{
+			VREditorMode->GetWorldInteraction().OnViewportInteractionInputAction().RemoveAll(this);
+			VREditorMode->GetWorldInteraction().OnViewportInteractionInputAction().AddRaw(this, &FEdModeFoliage::OnVRAction);
+			VREditorMode->GetWorldInteraction().OnViewportInteractionHoverUpdate().AddRaw(this, &FEdModeFoliage::OnVRHoverUpdate);
+		}
+	}
+}
+
+void FEdModeFoliage::OnVRHoverUpdate(FEditorViewportClient& ViewportClient, UViewportInteractor* Interactor, FVector& HoverImpactPoint, bool& bWasHandled)
+{
+	bool bHavePaintRay = false;
+	FVector PaintRayViewOrigin = FVector::ZeroVector;
+	FVector PaintRayStart = FVector::ZeroVector;
+	FVector PaintRayDirection = FVector::ZeroVector;
+
+	// Check if VR Editor is active
+	if (IVREditorModule::IsAvailable())
+	{
+		IVREditorMode* VREditorMode = static_cast<IVREditorMode*>(GetModeManager()->GetActiveMode(IVREditorModule::Get().GetVREditorModeID()));
+		if (VREditorMode != nullptr && VREditorMode->IsFullyInitialized())
+		{
+			// Grab interactor
+			for (UViewportInteractor* Interactor : VREditorMode->GetWorldInteraction().GetInteractors())
+			{
+				// Skip other interactors if we are painting with one
+				if (FoliageInteractor == nullptr || Interactor == FoliageInteractor)
+				{
+
+					// Go ahead and paint immediately
+					FVector LaserPointerStart, LaserPointerEnd;
+					if (Interactor->GetLaserPointer( /* Out */ LaserPointerStart, /* Out */ LaserPointerEnd))
+					{
+						const FVector LaserPointerDirection = (LaserPointerEnd - LaserPointerStart).GetSafeNormal();
+
+						PaintRayViewOrigin = VREditorMode->GetHeadTransform().GetLocation();
+						PaintRayStart = LaserPointerStart;
+						PaintRayDirection = LaserPointerDirection;
+
+						FoliageBrushTrace(&ViewportClient, PaintRayStart, PaintRayDirection);
+					}
+				}
+			}
+		}
+	}
+}
+
+void FEdModeFoliage::OnVRAction(class FEditorViewportClient& ViewportClient, UViewportInteractor* Interactor, const FViewportActionKeyInput& Action, bool& bOutIsInputCaptured, bool& bWasHandled)
+{
+	IVREditorMode* VREditorMode = static_cast<IVREditorMode*>(GetModeManager()->GetActiveMode(IVREditorModule::Get().GetVREditorModeID()));
+
+	if (VREditorMode != nullptr && Interactor != nullptr)
+	{
+		const UVREditorInteractor* VRInteractor =  Cast<UVREditorInteractor>(Interactor);
+		FoliageInteractor = Interactor;
+
+		// Consume both full press and light press
+		if (Action.ActionType == ViewportWorldActionTypes::SelectAndMove_LightlyPressed || Action.ActionType == ViewportWorldActionTypes::SelectAndMove)
+		{
+			static bool bIsLightlyPressed = false;
+			static bool bIsLightlyReleased = false;
+			static bool bIsFullyPressed = false;
+			static bool bIsFullyReleased = false;
+
+			if (Action.Event == IE_Pressed && !VRInteractor->IsHoveringOverUI())
+			{
+				bIsLightlyPressed = Action.ActionType == ViewportWorldActionTypes::SelectAndMove_LightlyPressed;
+				bIsFullyPressed = Action.ActionType == ViewportWorldActionTypes::SelectAndMove;
+
+				// Go ahead and paint immediately
+				FVector LaserPointerStart, LaserPointerEnd;
+				if (FoliageInteractor->GetLaserPointer( /* Out */ LaserPointerStart, /* Out */ LaserPointerEnd))
+				{
+					const FVector LaserPointerDirection = (LaserPointerEnd - LaserPointerStart).GetSafeNormal();
+					BrushTraceDirection = LaserPointerDirection;
+
+					// Only start painting if we're not dragging a widget handle
+					if (ViewportClient.GetCurrentWidgetAxis() == EAxisList::None)
+					{
+						bool bHandled = false;
+						if (UISettings.GetPaintToolSelected() || UISettings.GetReapplyToolSelected() || UISettings.GetLassoSelectToolSelected())
+						{
+							StartFoliageBrushTrace(&ViewportClient, Interactor);
+							FoliageBrushTrace(&ViewportClient, LaserPointerStart, LaserPointerDirection);
+							bHandled = true;
+						}
+						else if (!bHandled && UISettings.GetSelectToolSelected())
+						{
+							if (IsInputForRemovingBrushOn(&ViewportClient))
+							{
+								RemoveSelectedInstances(GetWorld());
+							}
+							else
+							{
+								SnapSelectedInstancesToGround(GetWorld());
+							}
+						}
+					}
+				}
+			}
+
+			// Stop current tracking if the user is no longer painting
+			else if (Action.Event == IE_Released && bIsPainting && FoliageInteractor && FoliageInteractor == Interactor)
+			{
+				bIsLightlyReleased = Action.ActionType == ViewportWorldActionTypes::SelectAndMove_LightlyPressed;
+				bIsFullyReleased = Action.ActionType == ViewportWorldActionTypes::SelectAndMove;
+
+				if ((bIsLightlyPressed && bIsLightlyReleased) || (bIsFullyPressed && bIsFullyReleased))
+				{
+					EndFoliageBrushTrace();
+					bIsLightlyPressed = false;
+					bIsLightlyReleased = false;
+					bIsFullyPressed = false;
+					bIsFullyReleased = false;
+
+				}
+			}
+		}
+	}
+}
+
 void FEdModeFoliage::PostUndo()
 {
 	FEdMode::PostUndo();
@@ -539,7 +702,6 @@ void FEdModeFoliage::OnObjectsReplaced(const TMap<UObject*, UObject*>& Replaceme
 	}
 }
 
-/** FEdMode: Called once per frame */
 void FEdModeFoliage::Tick(FEditorViewportClient* ViewportClient, float DeltaTime)
 {
 	if (bToolActive)
@@ -576,36 +738,57 @@ void FEdModeFoliage::Tick(FEditorViewportClient* ViewportClient, float DeltaTime
 	}
 }
 
-/** Trace under the mouse cursor and update brush position */
-void FEdModeFoliage::FoliageBrushTrace(FEditorViewportClient* ViewportClient, int32 MouseX, int32 MouseY)
+
+
+void FEdModeFoliage::StartFoliageBrushTrace(FEditorViewportClient* ViewportClient, class UViewportInteractor* Interactor)
+{
+	if (!bToolActive)
+	{
+		GEditor->BeginTransaction(NSLOCTEXT("UnrealEd", "FoliageMode_EditTransaction", "Foliage Editing"));
+		PreApplyBrush();
+		ApplyBrush(ViewportClient);
+		bToolActive = true;
+
+		if (Interactor != nullptr)
+		{
+			FoliageInteractor = Interactor;
+		}
+	}
+	bIsPainting = true;
+}
+
+void FEdModeFoliage::EndFoliageBrushTrace()
+{
+	GEditor->EndTransaction();
+	InstanceSnapshot.Empty();
+	LandscapeLayerCaches.Empty();
+	bToolActive = false;
+	bIsPainting = false;
+	bBrushTraceValid = false;
+
+	if (FoliageInteractor != nullptr)
+	{
+		FoliageInteractor = nullptr;
+	}
+}
+
+/** Trace and update brush position */
+void FEdModeFoliage::FoliageBrushTrace(FEditorViewportClient* ViewportClient, const FVector& InRayOrigin, const FVector& InRayDirection)
 {
 	bBrushTraceValid = false;
 	if (!ViewportClient->IsMovingCamera() && ViewportClient->IsVisible())
 	{
 		if (UISettings.GetPaintToolSelected() || UISettings.GetReapplyToolSelected() || UISettings.GetLassoSelectToolSelected())
 		{
-			// Compute a world space ray from the screen space mouse coordinates
-			FSceneViewFamilyContext ViewFamily(FSceneViewFamily::ConstructionValues(
-				ViewportClient->Viewport,
-				ViewportClient->GetScene(),
-				ViewportClient->EngineShowFlags)
-				.SetRealtimeUpdate(ViewportClient->IsRealtime()));
-			FSceneView* View = ViewportClient->CalcSceneView(&ViewFamily);
-			FViewportCursorLocation MouseViewportRay(View, ViewportClient, MouseX, MouseY);
-			BrushTraceDirection = MouseViewportRay.GetDirection();
-
-			FVector Start = MouseViewportRay.GetOrigin();
-			if (ViewportClient->IsOrtho())
-			{
-				Start += -WORLD_MAX * BrushTraceDirection;
-			}
-			FVector End = Start + WORLD_MAX * BrushTraceDirection;
+			const FVector TraceStart(InRayOrigin);
+			const FVector TraceEnd(InRayOrigin + InRayDirection * HALF_WORLD_MAX);
 
 			FHitResult Hit;
 			UWorld* World = ViewportClient->GetWorld();
 			static FName NAME_FoliageBrush = FName(TEXT("FoliageBrush"));
 			FFoliagePaintingGeometryFilter FilterFunc = FFoliagePaintingGeometryFilter(UISettings);
-			if (AInstancedFoliageActor::FoliageTrace(World, Hit, FDesiredFoliageInstance(Start, End), NAME_FoliageBrush, false, FilterFunc))
+			
+			if (AInstancedFoliageActor::FoliageTrace(World, Hit, FDesiredFoliageInstance(TraceStart, TraceEnd), NAME_FoliageBrush, false, FilterFunc))
 			{
 				UPrimitiveComponent* PrimComp = Hit.Component.Get();
 				if (CanPaint(PrimComp->GetComponentLevel()))
@@ -636,7 +819,29 @@ void FEdModeFoliage::FoliageBrushTrace(FEditorViewportClient* ViewportClient, in
  */
 bool FEdModeFoliage::MouseMove(FEditorViewportClient* ViewportClient, FViewport* Viewport, int32 MouseX, int32 MouseY)
 {
-	FoliageBrushTrace(ViewportClient, MouseX, MouseY);
+	// Use mouse capture if there's no other interactor currently tracing brush
+	IVREditorMode* VREditorMode = static_cast<IVREditorMode*>(GetModeManager()->GetActiveMode(IVREditorModule::Get().GetVREditorModeID()));
+	if (VREditorMode == nullptr)
+	{
+		// Compute a world space ray from the screen space mouse coordinates
+		FSceneViewFamilyContext ViewFamily(FSceneViewFamily::ConstructionValues(
+			ViewportClient->Viewport,
+			ViewportClient->GetScene(),
+			ViewportClient->EngineShowFlags)
+			.SetRealtimeUpdate(ViewportClient->IsRealtime()));
+		
+		FSceneView* View = ViewportClient->CalcSceneView(&ViewFamily);
+		FViewportCursorLocation MouseViewportRay(View, ViewportClient, MouseX, MouseY);
+		BrushTraceDirection = MouseViewportRay.GetDirection();
+
+		FVector BrushTraceStart = MouseViewportRay.GetOrigin();
+		if (ViewportClient->IsOrtho())
+		{
+			BrushTraceStart += -WORLD_MAX * BrushTraceDirection;
+		}
+
+		FoliageBrushTrace(ViewportClient, BrushTraceStart, BrushTraceDirection);
+	}
 	return false;
 }
 
@@ -652,7 +857,29 @@ bool FEdModeFoliage::MouseMove(FEditorViewportClient* ViewportClient, FViewport*
  */
 bool FEdModeFoliage::CapturedMouseMove(FEditorViewportClient* ViewportClient, FViewport* Viewport, int32 MouseX, int32 MouseY)
 {
-	FoliageBrushTrace(ViewportClient, MouseX, MouseY);
+	// Use mouse capture if there's no other interactor currently tracing brush
+	IVREditorMode* VREditorMode = static_cast<IVREditorMode*>(GetModeManager()->GetActiveMode(IVREditorModule::Get().GetVREditorModeID()));
+	if (VREditorMode == nullptr)
+	{
+		 //Compute a world space ray from the screen space mouse coordinates
+		FSceneViewFamilyContext ViewFamily(FSceneViewFamily::ConstructionValues(
+			ViewportClient->Viewport,
+			ViewportClient->GetScene(),
+			ViewportClient->EngineShowFlags)
+			.SetRealtimeUpdate(ViewportClient->IsRealtime()));
+
+		FSceneView* View = ViewportClient->CalcSceneView(&ViewFamily);
+		FViewportCursorLocation MouseViewportRay(View, ViewportClient, MouseX, MouseY);
+		BrushTraceDirection = MouseViewportRay.GetDirection();
+
+		FVector BrushTraceStart = MouseViewportRay.GetOrigin();
+		if (ViewportClient->IsOrtho())
+		{
+			BrushTraceStart += -WORLD_MAX * BrushTraceDirection;
+		}
+
+		FoliageBrushTrace(ViewportClient, BrushTraceStart, BrushTraceDirection);
+	}
 	return false;
 }
 
@@ -662,7 +889,7 @@ void FEdModeFoliage::GetRandomVectorInBrush(FVector& OutStart, FVector& OutEnd)
 	float Ru = (2.f * FMath::FRand() - 1.f);
 	float Rv = (2.f * FMath::FRand() - 1.f) * FMath::Sqrt(1.f - FMath::Square(Ru));
 
-	// find random point in circle thru brush location parallel to screen surface
+	// find random point in circle through brush location parallel to screen surface
 	FVector U, V;
 	BrushTraceDirection.FindBestAxisVectors(U, V);
 	FVector Point = Ru * U + Rv * V;
@@ -1896,8 +2123,9 @@ void FEdModeFoliage::ApplyBrush(FEditorViewportClient* ViewportClient)
 	
 	float BrushArea = PI * FMath::Square(UISettings.GetRadius());
 
-	// Tablet pressure
-	float Pressure = ViewportClient->Viewport->IsPenActive() ? ViewportClient->Viewport->GetTabletPressure() : 1.f;
+	// Tablet pressure or motion controller pressure
+	const UVREditorInteractor* VRInteractor = Cast<UVREditorInteractor>(FoliageInteractor);
+	const float Pressure = VRInteractor ? VRInteractor->GetSelectAndMoveTriggerValue() : ViewportClient->Viewport->IsPenActive() ? ViewportClient->Viewport->GetTabletPressure() : 1.f;
 
 	// Cache a copy of the world pointer
 	UWorld* World = ViewportClient->GetWorld();
@@ -1912,11 +2140,10 @@ void FEdModeFoliage::ApplyBrush(FEditorViewportClient* ViewportClient)
 		}
 
 		FSphere BrushSphere(BrushLocation, UISettings.GetRadius());
-					
+
 		if (UISettings.GetLassoSelectToolSelected())
 		{
-			// Shift unpaints
-			SelectInstancesForBrush(World, Settings, BrushSphere, !IsShiftDown(ViewportClient->Viewport));
+			SelectInstancesForBrush(World, Settings, BrushSphere, !IsInputForRemovingBrushOn(ViewportClient));
 		}
 		else if (UISettings.GetReapplyToolSelected())
 		{
@@ -1925,8 +2152,7 @@ void FEdModeFoliage::ApplyBrush(FEditorViewportClient* ViewportClient)
 		}
 		else if (UISettings.GetPaintToolSelected())
 		{
-			// Shift unpaints
-			if (IsShiftDown(ViewportClient->Viewport))
+			if (IsInputForRemovingBrushOn(ViewportClient))
 			{
 				int32 DesiredInstanceCount = FMath::RoundToInt(BrushArea * Settings->Density * UISettings.GetUnpaintDensity() / (1000.f*1000.f));
 					
@@ -1938,7 +2164,7 @@ void FEdModeFoliage::ApplyBrush(FEditorViewportClient* ViewportClient)
 				float DesiredInstanceCountFloat = BrushArea * Settings->Density * UISettings.GetPaintDensity() / (1000.f*1000.f);
 				// Allow a single instance with a random chance, if the brush is smaller than the density
 				int32 DesiredInstanceCount = DesiredInstanceCountFloat > 1.f ? FMath::RoundToInt(DesiredInstanceCountFloat) : FMath::FRand() < DesiredInstanceCountFloat ? 1 : 0;
-					
+				
 				AddInstancesForBrush(World, Settings, BrushSphere, DesiredInstanceCount, Pressure);
 			}
 		}
@@ -2560,6 +2786,19 @@ bool FEdModeFoliage::CanPaint(const UFoliageType* FoliageType, const ULevel* InL
 	return false;
 }
 
+bool FEdModeFoliage::IsInputForRemovingBrushOn(const FEditorViewportClient* ViewportClient) const
+{
+	IVREditorMode* VREditorMode = static_cast<IVREditorMode*>(GetModeManager()->GetActiveMode(IVREditorModule::Get().GetVREditorModeID()));
+	const UVREditorInteractor* VRInteractor = Cast<UVREditorInteractor>(FoliageInteractor);
+	bool bIsModifierPressed = false;
+	if (VREditorMode != nullptr && VRInteractor != nullptr)
+	{
+		bIsModifierPressed = VRInteractor->IsModifierPressed();
+	}
+
+	return IsShiftDown(ViewportClient->Viewport) || bIsModifierPressed;
+}
+
 UFoliageType* FEdModeFoliage::AddFoliageAsset(UObject* InAsset)
 {
 	UFoliageType* FoliageType = nullptr;
@@ -2804,7 +3043,7 @@ bool FEdModeFoliage::InputKey(FEditorViewportClient* ViewportClient, FViewport*
 	}
 	
 	bool bHandled = false;
-	if (UISettings.GetPaintToolSelected() || UISettings.GetReapplyToolSelected() || UISettings.GetLassoSelectToolSelected())
+	if ((UISettings.GetPaintToolSelected() || UISettings.GetReapplyToolSelected() || UISettings.GetLassoSelectToolSelected()) && FoliageInteractor == nullptr)
 	{
 		// Require Ctrl or not as per user preference
 		ELandscapeFoliageEditorControlType FoliageEditorControlType = GetDefault<ULevelEditorViewportSettings>()->FoliageEditorControlType;
@@ -2823,10 +3062,7 @@ bool FEdModeFoliage::InputKey(FEditorViewportClient* ViewportClient, FViewport*
 			{
 				if (!bToolActive)
 				{
-					GEditor->BeginTransaction(NSLOCTEXT("UnrealEd", "FoliageMode_EditTransaction", "Foliage Editing"));
-					PreApplyBrush();
-					ApplyBrush(ViewportClient);
-					bToolActive = true;
+					StartFoliageBrushTrace(ViewportClient, nullptr);
 
 					bHandled = true;
 				}
@@ -2837,10 +3073,7 @@ bool FEdModeFoliage::InputKey(FEditorViewportClient* ViewportClient, FViewport*
 		{
 			//Set the cursor position to that of the slate cursor so it wont snap back
 			Viewport->SetPreCaptureMousePosFromSlateCursor();
-			GEditor->EndTransaction();
-			InstanceSnapshot.Empty();
-			LandscapeLayerCaches.Empty();
-			bToolActive = false;
+			EndFoliageBrushTrace();
 			
 			bHandled = true;
 		}
@@ -2863,7 +3096,7 @@ bool FEdModeFoliage::InputKey(FEditorViewportClient* ViewportClient, FViewport*
 		}
 	}
 
-	if (!bHandled && (UISettings.GetLassoSelectToolSelected() || UISettings.GetSelectToolSelected()))
+	if (!bHandled && (UISettings.GetLassoSelectToolSelected() || UISettings.GetSelectToolSelected()) && FoliageInteractor == nullptr)
 	{
 		if (Event == IE_Pressed)
 		{
@@ -2978,7 +3211,7 @@ bool FEdModeFoliage::HandleClick(FEditorViewportClient* InViewportClient, HHitPr
 		{
 			GEditor->BeginTransaction(NSLOCTEXT("UnrealEd", "FoliageMode_EditTransaction", "Foliage Editing"));
 			
-			if (Click.IsShiftDown())
+			if (IsInputForRemovingBrushOn(InViewportClient))
 			{
 				ApplyPaintBucket_Remove(((HActor*)HitProxy)->Actor);
 			}
diff --git a/Engine/Source/Editor/FoliageEdit/Private/FoliageEdMode.h b/Engine/Source/Editor/FoliageEdit/Private/FoliageEdMode.h
--- a/Engine/Source/Editor/FoliageEdit/Private/FoliageEdMode.h
+++ b/Engine/Source/Editor/FoliageEdit/Private/FoliageEdMode.h
@@ -326,8 +326,14 @@ public:
 	/** Forces real-time perspective viewports */
 	void ForceRealTimeViewports(const bool bEnable, const bool bStoreCurrentState);
 
+	/** Start foliage tracing */
+	void StartFoliageBrushTrace(FEditorViewportClient* ViewportClient, class UViewportInteractor* Interactor);
+
+	/* End foliage tracing*/
+	void EndFoliageBrushTrace();
+
 	/** Trace under the mouse cursor and update brush position */
-	void FoliageBrushTrace(FEditorViewportClient* ViewportClient, int32 MouseX, int32 MouseY);
+	void FoliageBrushTrace(FEditorViewportClient* ViewportClient, const FVector& InRayOrigin, const FVector& InRayDirection);
 
 	/** Generate start/end points for a random trace inside the sphere brush.
 	    returns a line segment inside the sphere parallel to the view direction */
@@ -363,6 +369,9 @@ public:
 	/** Whether specified FoliageType can be painted into level */
 	static bool CanPaint(const UFoliageType* FoliageType, const ULevel* InLevel);
 
+	/** Should remove brushes based on user input (shift or modifier button pressed) */
+	bool IsInputForRemovingBrushOn(const FEditorViewportClient* ViewportClient) const;
+
 	/** Add a new asset (FoliageType or StaticMesh) */
 	UFoliageType* AddFoliageAsset(UObject* InAsset);
 	
@@ -402,6 +411,15 @@ public:
 	/** Add desired instances. Uses foliage settings to determine location/scale/rotation and whether instances should be ignored */
 	static void AddInstances(UWorld* InWorld, const TArray<FDesiredFoliageInstance>& DesiredInstances, const FFoliagePaintingGeometryFilter& OverrideGeometryFilter);
 
+	/** Called when an editor mode is entered or exited */
+	void OnEditorModeChanged(FEdMode* EditorMode, bool bEntered);
+
+	/** Called when the user presses a button on their motion controller device */
+	void OnVRAction(class FEditorViewportClient& ViewportClient, class UViewportInteractor* Interactor, const struct FViewportActionKeyInput& Action, bool& bOutIsInputCaptured, bool& bWasHandled);
+
+	/** Called on VR hovering */
+	void OnVRHoverUpdate(FEditorViewportClient& ViewportClient, UViewportInteractor* Interactor, FVector& HoverImpactPoint, bool& bWasHandled);
+
 	typedef TMap<FName, TMap<ULandscapeComponent*, TArray<uint8> > > LandscapeLayerCacheData;
 
 	FSimpleMulticastDelegate OnToolChanged;
@@ -517,5 +535,12 @@ private:
 	EColumnSortMode::Type			FoliageMeshListSortMode; 
 
 	FDelegateHandle OnActorSpawnedHandle;
+
+	/** Whether we're currently painting */
+	bool bIsPainting;
+
+	/** When painting in VR, this is the hand index that we're painting with.  Otherwise INDEX_NONE. */
+	class UViewportInteractor* FoliageInteractor;
+
 };
 
