[FoliageEdMode] Renamed variable to not hide declaration

diff --git a/Engine/Source/Editor/FoliageEdit/Private/FoliageEdMode.cpp b/Engine/Source/Editor/FoliageEdit/Private/FoliageEdMode.cpp
--- a/Engine/Source/Editor/FoliageEdit/Private/FoliageEdMode.cpp
+++ b/Engine/Source/Editor/FoliageEdit/Private/FoliageEdMode.cpp
@@ -463,10 +463,10 @@ void FEdModeFoliage::OnVRHoverUpdate(FEditorViewportClient& ViewportClient, UVie
 		if (VREditorMode != nullptr && VREditorMode->IsFullyInitialized())
 		{
 			// Grab interactor
-			for (UViewportInteractor* Interactor : VREditorMode->GetWorldInteraction().GetInteractors())
+			for (UViewportInteractor* ViewportInteractor : VREditorMode->GetWorldInteraction().GetInteractors())
 			{
 				// Skip other interactors if we are painting with one
-				if (FoliageInteractor == nullptr || Interactor == FoliageInteractor)
+				if (FoliageInteractor == nullptr || ViewportInteractor == FoliageInteractor)
 				{
 
 					// Go ahead and paint immediately
