Removed redundant code in FoliageEdMode

diff --git a/Engine/Source/Editor/FoliageEdit/Private/FoliageEdMode.cpp b/Engine/Source/Editor/FoliageEdit/Private/FoliageEdMode.cpp
--- a/Engine/Source/Editor/FoliageEdit/Private/FoliageEdMode.cpp
+++ b/Engine/Source/Editor/FoliageEdit/Private/FoliageEdMode.cpp
@@ -451,11 +451,6 @@ void FEdModeFoliage::OnEditorModeChanged(FEdMode* EditorMode, bool bEntered)
 
 void FEdModeFoliage::OnVRHoverUpdate(FEditorViewportClient& ViewportClient, UViewportInteractor* Interactor, FVector& HoverImpactPoint, bool& bWasHandled)
 {
-	bool bHavePaintRay = false;
-	FVector PaintRayViewOrigin = FVector::ZeroVector;
-	FVector PaintRayStart = FVector::ZeroVector;
-	FVector PaintRayDirection = FVector::ZeroVector;
-
 	// Check if VR Editor is active
 	if (IVREditorModule::IsAvailable())
 	{
@@ -475,11 +470,7 @@ void FEdModeFoliage::OnVRHoverUpdate(FEditorViewportClient& ViewportClient, UVie
 					{
 						const FVector LaserPointerDirection = (LaserPointerEnd - LaserPointerStart).GetSafeNormal();
 
-						PaintRayViewOrigin = VREditorMode->GetHeadTransform().GetLocation();
-						PaintRayStart = LaserPointerStart;
-						PaintRayDirection = LaserPointerDirection;
-
-						FoliageBrushTrace(&ViewportClient, PaintRayStart, PaintRayDirection);
+						FoliageBrushTrace(&ViewportClient, LaserPointerStart, LaserPointerDirection);
 					}
 				}
 			}
