VREditor: Fixed foliage fill not working

diff --git a/Engine/Source/Editor/FoliageEdit/Private/FoliageEdMode.cpp b/Engine/Source/Editor/FoliageEdit/Private/FoliageEdMode.cpp
--- a/Engine/Source/Editor/FoliageEdit/Private/FoliageEdMode.cpp
+++ b/Engine/Source/Editor/FoliageEdit/Private/FoliageEdMode.cpp
@@ -510,14 +510,34 @@ void FEdModeFoliage::OnVRAction(class FEditorViewportClient& ViewportClient, UVi
 					// Only start painting if we're not dragging a widget handle
 					if (ViewportClient.GetCurrentWidgetAxis() == EAxisList::None)
 					{
-						bool bHandled = false;
 						if (UISettings.GetPaintToolSelected() || UISettings.GetReapplyToolSelected() || UISettings.GetLassoSelectToolSelected())
 						{
 							StartFoliageBrushTrace(&ViewportClient, Interactor);
 							FoliageBrushTrace(&ViewportClient, LaserPointerStart, LaserPointerDirection);
-							bHandled = true;
 						}
-						else if (!bHandled && UISettings.GetSelectToolSelected())
+
+						// Fill a static mesh with foliage brush
+						else if (UISettings.GetPaintBucketToolSelected())
+						{
+							FHitResult HitResult = FoliageInteractor->GetHitResultFromLaserPointer();
+
+							if (HitResult.Actor.Get() != nullptr)
+							{
+								GEditor->BeginTransaction(NSLOCTEXT("UnrealEd", "FoliageMode_EditTransaction", "Foliage Editing"));
+
+								if (IsInputForRemovingBrushOn(&ViewportClient))
+								{
+									ApplyPaintBucket_Remove(HitResult.Actor.Get());
+								}
+								else
+								{
+									ApplyPaintBucket_Add(HitResult.Actor.Get());
+								}
+
+								GEditor->EndTransaction();
+							}
+						}
+						else if (UISettings.GetSelectToolSelected())
 						{
 							if (IsInputForRemovingBrushOn(&ViewportClient))
 							{
