VREditor: Fixed brush preview present while in foliage mode and hovering at UI

diff --git a/Engine/Source/Editor/FoliageEdit/Private/FoliageEdMode.cpp b/Engine/Source/Editor/FoliageEdit/Private/FoliageEdMode.cpp
--- a/Engine/Source/Editor/FoliageEdit/Private/FoliageEdMode.cpp
+++ b/Engine/Source/Editor/FoliageEdit/Private/FoliageEdMode.cpp
@@ -354,6 +354,8 @@ void FEdModeFoliage::Enter()
 	{
 		VREditorMode->GetWorldInteraction().OnViewportInteractionInputAction().RemoveAll(this);
 		VREditorMode->GetWorldInteraction().OnViewportInteractionInputAction().AddRaw(this, &FEdModeFoliage::OnVRAction);
+
+		VREditorMode->GetWorldInteraction().OnViewportInteractionHoverUpdate().RemoveAll(this);
 		VREditorMode->GetWorldInteraction().OnViewportInteractionHoverUpdate().AddRaw(this, &FEdModeFoliage::OnVRHoverUpdate);
 	}
 }
@@ -368,6 +370,7 @@ void FEdModeFoliage::Exit()
 		if (VREditorMode != nullptr)
 		{
 			VREditorMode->GetWorldInteraction().OnViewportInteractionInputAction().RemoveAll(this);
+			VREditorMode->GetWorldInteraction().OnViewportInteractionHoverUpdate().RemoveAll(this);
 			FoliageInteractor = nullptr;
 		}
 	}
@@ -444,6 +447,8 @@ void FEdModeFoliage::OnEditorModeChanged(FEdMode* EditorMode, bool bEntered)
 		{
 			VREditorMode->GetWorldInteraction().OnViewportInteractionInputAction().RemoveAll(this);
 			VREditorMode->GetWorldInteraction().OnViewportInteractionInputAction().AddRaw(this, &FEdModeFoliage::OnVRAction);
+
+			VREditorMode->GetWorldInteraction().OnViewportInteractionHoverUpdate().RemoveAll(this);
 			VREditorMode->GetWorldInteraction().OnViewportInteractionHoverUpdate().AddRaw(this, &FEdModeFoliage::OnVRHoverUpdate);
 		}
 	}
@@ -460,10 +465,24 @@ void FEdModeFoliage::OnVRHoverUpdate(FEditorViewportClient& ViewportClient, UVie
 			// Grab interactor
 			for (UViewportInteractor* ViewportInteractor : VREditorMode->GetWorldInteraction().GetInteractors())
 			{
-				// Skip other interactors if we are painting with one
-				if (FoliageInteractor == nullptr || ViewportInteractor == FoliageInteractor)
+				// Check if we're hovering over UI. If so, stop painting so we don't display the preview brush sphere
+				const UVREditorInteractor* VRInteractor = Cast<UVREditorInteractor>(ViewportInteractor);
+				const bool bIsHoveringOverUIVR = VRInteractor->IsHoveringOverUI();
+				if (bIsHoveringOverUIVR)
+				{
+					EndFoliageBrushTrace();
+				}
+				else
 				{
+					if (UISettings.GetPaintToolSelected() || UISettings.GetReapplyToolSelected() || UISettings.GetLassoSelectToolSelected())
+					{
+						StartFoliageBrushTrace(&ViewportClient, ViewportInteractor);
+					}
+				}
 
+				// Skip other interactors if we are painting with one
+				if (ViewportInteractor && ViewportInteractor == FoliageInteractor)
+				{
 					// Go ahead and paint immediately
 					FVector LaserPointerStart, LaserPointerEnd;
 					if (Interactor->GetLaserPointer( /* Out */ LaserPointerStart, /* Out */ LaserPointerEnd))
