VREditor: Small clean for a foliage function signature

diff --git a/Engine/Source/Editor/FoliageEdit/Private/FoliageEdMode.cpp b/Engine/Source/Editor/FoliageEdit/Private/FoliageEdMode.cpp
--- a/Engine/Source/Editor/FoliageEdit/Private/FoliageEdMode.cpp
+++ b/Engine/Source/Editor/FoliageEdit/Private/FoliageEdMode.cpp
@@ -3092,7 +3092,7 @@ bool FEdModeFoliage::InputKey(FEditorViewportClient* ViewportClient, FViewport*
 			{
 				if (!bToolActive)
 				{
-					StartFoliageBrushTrace(ViewportClient, nullptr);
+					StartFoliageBrushTrace(ViewportClient);
 
 					bHandled = true;
 				}
diff --git a/Engine/Source/Editor/FoliageEdit/Private/FoliageEdMode.h b/Engine/Source/Editor/FoliageEdit/Private/FoliageEdMode.h
--- a/Engine/Source/Editor/FoliageEdit/Private/FoliageEdMode.h
+++ b/Engine/Source/Editor/FoliageEdit/Private/FoliageEdMode.h
@@ -327,7 +327,7 @@ public:
 	void ForceRealTimeViewports(const bool bEnable, const bool bStoreCurrentState);
 
 	/** Start foliage tracing */
-	void StartFoliageBrushTrace(FEditorViewportClient* ViewportClient, class UViewportInteractor* Interactor);
+	void StartFoliageBrushTrace(FEditorViewportClient* ViewportClient, class UViewportInteractor* Interactor = nullptr);
 
 	/* End foliage tracing*/
 	void EndFoliageBrushTrace();
