VREditor: Fixed foliage select tool not properly select foliage instanced mesh

diff --git a/Engine/Source/Editor/FoliageEdit/Private/FoliageEdMode.cpp b/Engine/Source/Editor/FoliageEdit/Private/FoliageEdMode.cpp
--- a/Engine/Source/Editor/FoliageEdit/Private/FoliageEdMode.cpp
+++ b/Engine/Source/Editor/FoliageEdit/Private/FoliageEdMode.cpp
@@ -509,16 +509,8 @@ void FEdModeFoliage::OnVRAction(class FEditorViewportClient& ViewportClient, UVi
 		// Consume both full press and light press
 		if (Action.ActionType == ViewportWorldActionTypes::SelectAndMove_LightlyPressed || Action.ActionType == ViewportWorldActionTypes::SelectAndMove)
 		{
-			static bool bIsLightlyPressed = false;
-			static bool bIsLightlyReleased = false;
-			static bool bIsFullyPressed = false;
-			static bool bIsFullyReleased = false;
-
 			if (Action.Event == IE_Pressed && !VRInteractor->IsHoveringOverUI())
 			{
-				bIsLightlyPressed = Action.ActionType == ViewportWorldActionTypes::SelectAndMove_LightlyPressed;
-				bIsFullyPressed = Action.ActionType == ViewportWorldActionTypes::SelectAndMove;
-
 				// Go ahead and paint immediately
 				FVector LaserPointerStart, LaserPointerEnd;
 				if (FoliageInteractor->GetLaserPointer( /* Out */ LaserPointerStart, /* Out */ LaserPointerEnd))
@@ -536,7 +528,7 @@ void FEdModeFoliage::OnVRAction(class FEditorViewportClient& ViewportClient, UVi
 						}
 
 						// Fill a static mesh with foliage brush
-						else if (UISettings.GetPaintBucketToolSelected())
+						else if (UISettings.GetPaintBucketToolSelected() || UISettings.GetReapplyPaintBucketToolSelected())
 						{
 							FHitResult HitResult = FoliageInteractor->GetHitResultFromLaserPointer();
 
@@ -544,7 +536,7 @@ void FEdModeFoliage::OnVRAction(class FEditorViewportClient& ViewportClient, UVi
 							{
 								GEditor->BeginTransaction(NSLOCTEXT("UnrealEd", "FoliageMode_EditTransaction", "Foliage Editing"));
 
-								if (IsInputForRemovingBrushOn(&ViewportClient))
+								if (IsModifierButtonPressed(&ViewportClient))
 								{
 									ApplyPaintBucket_Remove(HitResult.Actor.Get());
 								}
@@ -558,14 +550,26 @@ void FEdModeFoliage::OnVRAction(class FEditorViewportClient& ViewportClient, UVi
 						}
 						else if (UISettings.GetSelectToolSelected())
 						{
-							if (IsInputForRemovingBrushOn(&ViewportClient))
-							{
-								RemoveSelectedInstances(GetWorld());
-							}
-							else
+							FHitResult HitResult = FoliageInteractor->GetHitResultFromLaserPointer();
+							
+							GEditor->BeginTransaction(NSLOCTEXT("UnrealEd", "FoliageMode_EditTransaction", "Foliage Editing"));
+
+							if (HitResult.GetActor() != nullptr)
 							{
-								SnapSelectedInstancesToGround(GetWorld());
+								// Clear all currently selected instances
+								SelectInstances(ViewportClient.GetWorld(), false);
+								for (auto& FoliageMeshUI : FoliageMeshList)
+								{
+									UFoliageType* Settings = FoliageMeshUI->Settings;
+									SelectInstanceAtLocation(ViewportClient.GetWorld(), Settings, HitResult.ImpactPoint, !IsModifierButtonPressed(&ViewportClient));
+								}
 							}
+
+							GEditor->EndTransaction();
+
+							// @todo vreditor: we currently don't have a key mapping scheme to snap selected instances to ground 
+							// SnapSelectedInstancesToGround(GetWorld());
+
 						}
 					}
 				}
@@ -574,18 +578,7 @@ void FEdModeFoliage::OnVRAction(class FEditorViewportClient& ViewportClient, UVi
 			// Stop current tracking if the user is no longer painting
 			else if (Action.Event == IE_Released && bIsPainting && FoliageInteractor && FoliageInteractor == Interactor)
 			{
-				bIsLightlyReleased = Action.ActionType == ViewportWorldActionTypes::SelectAndMove_LightlyPressed;
-				bIsFullyReleased = Action.ActionType == ViewportWorldActionTypes::SelectAndMove;
-
-				if ((bIsLightlyPressed && bIsLightlyReleased) || (bIsFullyPressed && bIsFullyReleased))
-				{
-					EndFoliageBrushTrace();
-					bIsLightlyPressed = false;
-					bIsLightlyReleased = false;
-					bIsFullyPressed = false;
-					bIsFullyReleased = false;
-
-				}
+				EndFoliageBrushTrace();
 			}
 		}
 	}
@@ -1463,6 +1456,26 @@ void FEdModeFoliage::RemoveInstancesForBrush(UWorld* InWorld, const UFoliageType
 	}
 }
 
+
+void FEdModeFoliage::SelectInstanceAtLocation(UWorld* InWorld, const UFoliageType* Settings, const FVector& Location, bool bSelect)
+{
+	for (FFoliageMeshInfoIterator It(InWorld, Settings); It; ++It)
+	{
+		FFoliageMeshInfo* MeshInfo = (*It);
+		AInstancedFoliageActor* IFA = It.GetActor();
+
+		int32 Instance;
+		bool bResult;
+		MeshInfo->GetInstanceAtLocation(Location, Instance, bResult);
+		if (bResult)
+		{
+		        TArray<int32> Instances;
+			Instances.Add(Instance);
+			MeshInfo->SelectInstances(IFA, bSelect, Instances);
+		}
+	}
+}
+
 void FEdModeFoliage::SelectInstancesForBrush(UWorld* InWorld, const UFoliageType* Settings, const FSphere& BrushSphere, bool bSelect)
 {
 	for (FFoliageMeshInfoIterator It(InWorld, Settings); It; ++It)
@@ -2173,7 +2186,7 @@ void FEdModeFoliage::ApplyBrush(FEditorViewportClient* ViewportClient)
 
 		if (UISettings.GetLassoSelectToolSelected())
 		{
-			SelectInstancesForBrush(World, Settings, BrushSphere, !IsInputForRemovingBrushOn(ViewportClient));
+			SelectInstancesForBrush(World, Settings, BrushSphere, !IsModifierButtonPressed(ViewportClient));
 		}
 		else if (UISettings.GetReapplyToolSelected())
 		{
@@ -2182,7 +2195,7 @@ void FEdModeFoliage::ApplyBrush(FEditorViewportClient* ViewportClient)
 		}
 		else if (UISettings.GetPaintToolSelected())
 		{
-			if (IsInputForRemovingBrushOn(ViewportClient))
+			if (IsModifierButtonPressed(ViewportClient))
 			{
 				int32 DesiredInstanceCount = FMath::RoundToInt(BrushArea * Settings->Density * UISettings.GetUnpaintDensity() / (1000.f*1000.f));
 					
@@ -2816,7 +2829,7 @@ bool FEdModeFoliage::CanPaint(const UFoliageType* FoliageType, const ULevel* InL
 	return false;
 }
 
-bool FEdModeFoliage::IsInputForRemovingBrushOn(const FEditorViewportClient* ViewportClient) const
+bool FEdModeFoliage::IsModifierButtonPressed(const FEditorViewportClient* ViewportClient) const
 {
 	IVREditorMode* VREditorMode = static_cast<IVREditorMode*>(GetModeManager()->GetActiveMode(IVREditorModule::Get().GetVREditorModeID()));
 	const UVREditorInteractor* VRInteractor = Cast<UVREditorInteractor>(FoliageInteractor);
@@ -3241,7 +3254,7 @@ bool FEdModeFoliage::HandleClick(FEditorViewportClient* InViewportClient, HHitPr
 		{
 			GEditor->BeginTransaction(NSLOCTEXT("UnrealEd", "FoliageMode_EditTransaction", "Foliage Editing"));
 			
-			if (IsInputForRemovingBrushOn(InViewportClient))
+			if (IsModifierButtonPressed(InViewportClient))
 			{
 				ApplyPaintBucket_Remove(((HActor*)HitProxy)->Actor);
 			}
diff --git a/Engine/Source/Editor/FoliageEdit/Private/FoliageEdMode.h b/Engine/Source/Editor/FoliageEdit/Private/FoliageEdMode.h
--- a/Engine/Source/Editor/FoliageEdit/Private/FoliageEdMode.h
+++ b/Engine/Source/Editor/FoliageEdit/Private/FoliageEdMode.h
@@ -369,8 +369,8 @@ public:
 	/** Whether specified FoliageType can be painted into level */
 	static bool CanPaint(const UFoliageType* FoliageType, const ULevel* InLevel);
 
-	/** Should remove brushes based on user input (shift or modifier button pressed) */
-	bool IsInputForRemovingBrushOn(const FEditorViewportClient* ViewportClient) const;
+	/** Shift or modifier button pressed */
+	bool IsModifierButtonPressed(const FEditorViewportClient* ViewportClient) const;
 
 	/** Add a new asset (FoliageType or StaticMesh) */
 	UFoliageType* AddFoliageAsset(UObject* InAsset);
@@ -467,6 +467,9 @@ private:
 	/** Select instances inside the brush. */
 	void SelectInstancesForBrush(UWorld* InWorld, const UFoliageType* Settings, const FSphere& BrushSphere, bool bSelect);
 	
+	/** Select instance closest to the brush. */
+	void SelectInstanceAtLocation(UWorld* InWorld, const UFoliageType* Settings, const FVector& BrushLocation, bool bSelect);
+
 	/** Set/Clear selection for all foliage instances */
 	void SelectInstances(UWorld* InWorld, bool bSelect);
 	
diff --git a/Engine/Source/Runtime/Foliage/Private/InstancedFoliage.cpp b/Engine/Source/Runtime/Foliage/Private/InstancedFoliage.cpp
--- a/Engine/Source/Runtime/Foliage/Private/InstancedFoliage.cpp
+++ b/Engine/Source/Runtime/Foliage/Private/InstancedFoliage.cpp
@@ -1194,6 +1194,27 @@ void FFoliageMeshInfo::GetInstancesInsideSphere(const FSphere& Sphere, TArray<in
 	}
 }
 
+void FFoliageMeshInfo::GetInstanceAtLocation(const FVector& Location, int32& OutInstance, bool& bOutSucess)
+{
+	auto TempInstances = InstanceHash->GetInstancesOverlappingBox(FBox::BuildAABB(Location, FVector(KINDA_SMALL_NUMBER)));
+
+	float ShortestDistance = MAX_FLT;
+	OutInstance = -1;
+
+	for (int32 Idx : TempInstances)
+	{
+		FVector InstanceLocation = Instances[Idx].Location;
+		float DistanceSquared = FVector::DistSquared(InstanceLocation, Location);
+		if (DistanceSquared < ShortestDistance)
+		{
+			ShortestDistance = DistanceSquared;
+			OutInstance = Idx;
+		}
+	}
+
+	bOutSucess = OutInstance != -1;
+}
+
 // Returns whether or not there is are any instances overlapping the sphere specified
 bool FFoliageMeshInfo::CheckForOverlappingSphere(const FSphere& Sphere)
 {
diff --git a/Engine/Source/Runtime/Foliage/Public/InstancedFoliage.h b/Engine/Source/Runtime/Foliage/Public/InstancedFoliage.h
--- a/Engine/Source/Runtime/Foliage/Public/InstancedFoliage.h
+++ b/Engine/Source/Runtime/Foliage/Public/InstancedFoliage.h
@@ -198,6 +198,7 @@ struct FFoliageMeshInfo
 	FOLIAGE_API void PostUpdateInstances(AInstancedFoliageActor* InIFA, const TArray<int32>& InInstancesUpdated, bool bReAddToHash = false);
 	FOLIAGE_API void DuplicateInstances(AInstancedFoliageActor* InIFA, UFoliageType* InSettings, const TArray<int32>& InInstancesToDuplicate);
 	FOLIAGE_API void GetInstancesInsideSphere(const FSphere& Sphere, TArray<int32>& OutInstances);
+	FOLIAGE_API void GetInstanceAtLocation(const FVector& Location, int32& OutInstance, bool& bOutSucess);
 	FOLIAGE_API bool CheckForOverlappingSphere(const FSphere& Sphere);
 	FOLIAGE_API bool CheckForOverlappingInstanceExcluding(int32 TestInstanceIdx, float Radius, TSet<int32>& ExcludeInstances);
 	
