Fix for crash on editor close after VR Foliage Editing.

diff --git a/Engine/Source/Editor/FoliageEdit/Private/FoliageEdMode.cpp b/Engine/Source/Editor/FoliageEdit/Private/FoliageEdMode.cpp
--- a/Engine/Source/Editor/FoliageEdit/Private/FoliageEdMode.cpp
+++ b/Engine/Source/Editor/FoliageEdit/Private/FoliageEdMode.cpp
@@ -452,13 +452,6 @@ void FEdModeFoliage::OnVRHoverUpdate(FEditorViewportClient& ViewportClient, UVie
 				{
 					EndFoliageBrushTrace();
 				}
-				else
-				{
-					if (UISettings.GetPaintToolSelected() || UISettings.GetReapplyToolSelected() || UISettings.GetLassoSelectToolSelected())
-					{
-						StartFoliageBrushTrace(&ViewportClient, ViewportInteractor);
-					}
-				}
 
 				// Skip other interactors if we are painting with one
 				if (ViewportInteractor && ViewportInteractor == FoliageInteractor)
@@ -760,11 +753,6 @@ void FEdModeFoliage::Tick(FEditorViewportClient* ViewportClient, float DeltaTime
 					}
 					else if (UISettings.GetPaintToolSelected() || UISettings.GetReapplyToolSelected() || UISettings.GetLassoSelectToolSelected())
 					{
-						if (UISettings.GetPaintToolSelected())
-						{
-							StartFoliageBrushTrace(ViewportClient, Interactor);
-						}
-
 						// Go ahead and paint immediately
 						FVector LaserPointerStart, LaserPointerEnd;
 						if (Interactor->GetLaserPointer( /* Out */ LaserPointerStart, /* Out */ LaserPointerEnd))
