Foliage mode now ignores trigger input if the radial menu is present on the interactor. Fixes undo/redo not working on foliage mode.

diff --git a/Engine/Source/Editor/FoliageEdit/Private/FoliageEdMode.cpp b/Engine/Source/Editor/FoliageEdit/Private/FoliageEdMode.cpp
--- a/Engine/Source/Editor/FoliageEdit/Private/FoliageEdMode.cpp
+++ b/Engine/Source/Editor/FoliageEdit/Private/FoliageEdMode.cpp
@@ -486,9 +486,9 @@ void FEdModeFoliage::OnVRAction(class FEditorViewportClient& ViewportClient, UVi
 		// Only allow light press
 		Interactor->SetAllowTriggerFullPress(false);
 		// Consume light press
-		if (Action.ActionType == ViewportWorldActionTypes::SelectAndMove_LightlyPressed)
+		if ((Action.ActionType == ViewportWorldActionTypes::SelectAndMove_LightlyPressed) && (!VREditorMode->IsShowingRadialMenu(VRInteractor)))
 		{
-			if (Action.Event == IE_Pressed && !VRInteractor->IsHoveringOverPriorityType())
+			if ((Action.Event == IE_Pressed) && (!VRInteractor->IsHoveringOverPriorityType()))
 			{
 				// Go ahead and paint immediately
 				FVector LaserPointerStart, LaserPointerEnd;
@@ -556,7 +556,7 @@ void FEdModeFoliage::OnVRAction(class FEditorViewportClient& ViewportClient, UVi
 			}
 
 			// Stop current tracking if the user is no longer painting
-			else if (Action.Event == IE_Released && FoliageInteractor && FoliageInteractor == Interactor)
+			else if (Action.Event == IE_Released && FoliageInteractor && (FoliageInteractor == Interactor))
 			{
 				EndFoliageBrushTrace();
 				FoliageInteractor = nullptr;
diff --git a/Engine/Source/Editor/VREditor/Public/VREditorMode.h b/Engine/Source/Editor/VREditor/Public/VREditorMode.h
--- a/Engine/Source/Editor/VREditor/Public/VREditorMode.h
+++ b/Engine/Source/Editor/VREditor/Public/VREditorMode.h
@@ -194,6 +194,9 @@ public:
 		return *UISystem;
 	}
 
+	/** Lets other modules know if the radial menu is visible on a given interactor so input should be handled differently */
+	bool IsShowingRadialMenu( const class UVREditorInteractor* Interactor ) const;
+
 	/** Gets the viewport that VR Mode is activated in.  Even though editor modes are available in all
 	    level viewports simultaneously, only one viewport is "possessed" by the HMD.  Generally try to avoid using
 		this function and instead use the ViewportClient that is passed around through various FEdMode overrides */
diff --git a/Engine/Source/Editor/VREditor/UI/VREditorUISystem.cpp b/Engine/Source/Editor/VREditor/UI/VREditorUISystem.cpp
--- a/Engine/Source/Editor/VREditor/UI/VREditorUISystem.cpp
+++ b/Engine/Source/Editor/VREditor/UI/VREditorUISystem.cpp
@@ -1218,7 +1218,7 @@ void UVREditorUISystem::ShowEditorUIPanel( AVREditorFloatingUI* Panel, UVREditor
 }
 
 
-bool UVREditorUISystem::IsShowingRadialMenu( UVREditorInteractor* Interactor ) const
+bool UVREditorUISystem::IsShowingRadialMenu(const UVREditorInteractor* Interactor ) const
 {
 	const EControllerHand DockedToHand = QuickRadialMenu->GetDockedTo() == AVREditorFloatingUI::EDockedTo::LeftArm ? EControllerHand::Left : EControllerHand::Right;
 	const UVREditorInteractor* DockedToInteractor = GetOwner().GetHandInteractor( DockedToHand );
diff --git a/Engine/Source/Editor/VREditor/UI/VREditorUISystem.h b/Engine/Source/Editor/VREditor/UI/VREditorUISystem.h
--- a/Engine/Source/Editor/VREditor/UI/VREditorUISystem.h
+++ b/Engine/Source/Editor/VREditor/UI/VREditorUISystem.h
@@ -142,7 +142,7 @@ public:
 	void ShowEditorUIPanel( class AVREditorFloatingUI* Panel, UVREditorInteractor* Interactor, const bool bShouldShow, const bool OnHand = false, const bool bRefreshQuickMenu = true, const bool bPlaySound = true );
 
 	/** Returns true if the radial menu is visible on this hand */
-	bool IsShowingRadialMenu( UVREditorInteractor* Interactor ) const;
+	bool IsShowingRadialMenu(const UVREditorInteractor* Interactor ) const;
 
 	/** Tries to spawn the radial menu (if the specified hand isn't doing anything else) */
 	void TryToSpawnRadialMenu( UVREditorInteractor* Interactor, const bool bForceOverUI );
diff --git a/Engine/Source/Editor/VREditor/VREditorMode.cpp b/Engine/Source/Editor/VREditor/VREditorMode.cpp
--- a/Engine/Source/Editor/VREditor/VREditorMode.cpp
+++ b/Engine/Source/Editor/VREditor/VREditorMode.cpp
@@ -792,6 +792,11 @@ void UVREditorMode::DestroyTransientActor( AActor* Actor ) const
 	WorldInteraction->DestroyTransientActor( Actor );
 }
 
+bool UVREditorMode::IsShowingRadialMenu(const UVREditorInteractor* Interactor) const
+{
+	return UISystem->IsShowingRadialMenu(Interactor);
+}
+
 const SLevelViewport& UVREditorMode::GetLevelViewportPossessedForVR() const
 {
 	return *VREditorLevelViewportWeakPtr.Pin();
