[NO NON-PII COMMIT MESSAGE]

diff --git a/Engine/Source/Editor/FoliageEdit/Private/FoliageEdMode.cpp b/Engine/Source/Editor/FoliageEdit/Private/FoliageEdMode.cpp
--- a/Engine/Source/Editor/FoliageEdit/Private/FoliageEdMode.cpp
+++ b/Engine/Source/Editor/FoliageEdit/Private/FoliageEdMode.cpp
@@ -760,7 +760,7 @@ void FEdModeFoliage::StartFoliageBrushTrace(FEditorViewportClient* ViewportClien
 		PreApplyBrush();
 		ApplyBrush(ViewportClient);
 
-		if (UISettings.GetIsInQuickSingleInstantiationMode() || UISettings.GetIsInSingleInstantiationMode())
+		if (UISettings.IsInAnySingleInstantiationMode())
 		{
 			EndFoliageBrushTrace();
 		}
@@ -1763,6 +1763,11 @@ void FEdModeFoliage::SelectInvalidInstances(const UFoliageType* Settings)
 
 void FEdModeFoliage::AdjustBrushRadius(float Adjustment)
 {
+	if (UISettings.IsInAnySingleInstantiationMode())
+	{
+		return;
+	}
+
 	const float CurrentBrushRadius = UISettings.GetRadius();
 
 	if (Adjustment > 0.f)
@@ -1773,7 +1778,6 @@ void FEdModeFoliage::AdjustBrushRadius(float Adjustment)
 	{
 		UISettings.SetRadius(FMath::Max(CurrentBrushRadius + Adjustment, 0.f));
 	}
-
 }
 
 void FEdModeFoliage::ReapplyInstancesForBrush(UWorld* InWorld, const UFoliageType* Settings, const FSphere& BrushSphere, float Pressure)
@@ -2206,7 +2210,7 @@ void FEdModeFoliage::ApplyBrush(FEditorViewportClient* ViewportClient)
 			}
 			else
 			{
-				if (UISettings.GetIsInQuickSingleInstantiationMode() || UISettings.GetIsInSingleInstantiationMode())
+				if (UISettings.IsInAnySingleInstantiationMode())
 				{
 					AddSingleInstanceForBrush(World, Settings, Pressure);
 				}
diff --git a/Engine/Source/Editor/FoliageEdit/Private/FoliageEdMode.h b/Engine/Source/Editor/FoliageEdit/Private/FoliageEdMode.h
--- a/Engine/Source/Editor/FoliageEdit/Private/FoliageEdMode.h
+++ b/Engine/Source/Editor/FoliageEdit/Private/FoliageEdMode.h
@@ -64,8 +64,8 @@ struct FFoliageUISettings
 	bool GetReapplyPaintBucketToolSelected() const { return bReapplyPaintBucketToolSelected ? true : false; }
 	void SetReapplyPaintBucketToolSelected(bool InbReapplyPaintBucketToolSelected) { bReapplyPaintBucketToolSelected = InbReapplyPaintBucketToolSelected; }
 
-	float GetRadius() const { return (GetIsInSingleInstantiationMode() || GetIsInQuickSingleInstantiationMode()) ? SingleInstanceModeBrushSize : Radius; }
-	void SetRadius(float InRadius) { Radius = InRadius; }
+	float GetRadius() const { return (IsInAnySingleInstantiationMode()) ? SingleInstanceModeBrushSize : Radius; }
+	void SetRadius(float InRadius) { if (!IsInAnySingleInstantiationMode()) Radius = InRadius; }
 	float GetPaintDensity() const { return PaintDensity; }
 	void SetPaintDensity(float InPaintDensity) { PaintDensity = InPaintDensity; }
 	float GetUnpaintDensity() const { return UnpaintDensity; }
@@ -81,6 +81,8 @@ struct FFoliageUISettings
 	bool GetFilterTranslucent() const { return bFilterTranslucent; }
 	void SetFilterTranslucent(bool InbFilterTranslucent) { bFilterTranslucent = InbFilterTranslucent; }
 
+	bool IsInAnySingleInstantiationMode() const { return GetIsInSingleInstantiationMode() || GetIsInQuickSingleInstantiationMode(); }
+
 	bool GetIsInSingleInstantiationMode() const { return IsInSingleInstantiationMode; }
 	void SetIsInSingleInstantiationMode(bool InIsInSingleInstantiationMode) { IsInSingleInstantiationMode = InIsInSingleInstantiationMode; }
 
diff --git a/Engine/Source/Editor/FoliageEdit/Private/SFoliageEdit.cpp b/Engine/Source/Editor/FoliageEdit/Private/SFoliageEdit.cpp
--- a/Engine/Source/Editor/FoliageEdit/Private/SFoliageEdit.cpp
+++ b/Engine/Source/Editor/FoliageEdit/Private/SFoliageEdit.cpp
@@ -592,7 +592,7 @@ TOptional<float> SFoliageEdit::GetRadius() const
 
 bool SFoliageEdit::IsEnabled_BrushSize() const
 {
-	return !FoliageEditMode->UISettings.GetIsInSingleInstantiationMode() && !FoliageEditMode->UISettings.GetIsInQuickSingleInstantiationMode();
+	return !FoliageEditMode->UISettings.IsInAnySingleInstantiationMode();
 }
 
 void SFoliageEdit::SetPaintDensity(float InDensity)
@@ -607,7 +607,7 @@ TOptional<float> SFoliageEdit::GetPaintDensity() const
 
 bool SFoliageEdit::IsEnabled_PaintDensity() const
 {
-	return !FoliageEditMode->UISettings.GetIsInSingleInstantiationMode() && !FoliageEditMode->UISettings.GetIsInQuickSingleInstantiationMode();
+	return !FoliageEditMode->UISettings.IsInAnySingleInstantiationMode();
 }
 
 void SFoliageEdit::SetEraseDensity(float InDensity)
@@ -622,7 +622,7 @@ TOptional<float> SFoliageEdit::GetEraseDensity() const
 
 bool SFoliageEdit::IsEnabled_EraseDensity() const
 {
-	return !FoliageEditMode->UISettings.GetIsInSingleInstantiationMode() && !FoliageEditMode->UISettings.GetIsInQuickSingleInstantiationMode();
+	return !FoliageEditMode->UISettings.IsInAnySingleInstantiationMode();
 }
 
 EVisibility SFoliageEdit::GetVisibility_SelectionOptions() const
