VR Editor: Fix Foliage Paint mesh does not render in desktop after painting in VREditor.

diff --git a/Engine/Source/Editor/FoliageEdit/Private/FoliageEdMode.cpp b/Engine/Source/Editor/FoliageEdit/Private/FoliageEdMode.cpp
--- a/Engine/Source/Editor/FoliageEdit/Private/FoliageEdMode.cpp
+++ b/Engine/Source/Editor/FoliageEdit/Private/FoliageEdMode.cpp
@@ -382,6 +382,9 @@ void FEdModeFoliage::Enter()
 
 		SetBrushOpacity(VREd::FoliageOpacity->GetFloat());
 	}
+
+	// Make sure the brush is visible.
+	SphereBrushComponent->SetVisibility(true);
 }
 
 /** FEdMode: Called when the mode is exited */
