Copying //UE4/Dev-VREditor to //UE4/Dev-Main (Source: //UE4/Dev-VREditor @ 3352034)

==========================
MAJOR FEATURES + CHANGES
==========================

Change 3241752 on 2016/12/21 by [REDACTED_PERSON]

	VR Editor: New Create Sequence button on the Quick Menu allows you to create a new level sequence in VR. Sequencer UI accessible from this new asset as well as opening sequence assets from the Content Browser. In VR Editing mode, all autokeying is enabled. Settings are restored when leaving VR Editing mode. Meshes are placed in the world to represent keys and frames (VR and non-VR editing mode in this change). Transforming key meshes pushes the transform change back into the correct Sequencer track.

Change 3241754 on 2016/12/21 by [REDACTED_PERSON]

	Updating VR  editor interaction code to not allow selecting primitive components that have been set as non-selectable

Change 3241755 on 2016/12/21 by [REDACTED_PERSON]

	Moving SVirtualWindow into Slate module

Change 3242831 on 2016/12/22 by [REDACTED_PERSON]

	Changing TArray<TTuple>s to TArray<Struct> to fix clang crash, improve code readability

Change 3245240 on 2017/01/03 by [REDACTED_PERSON]

	ViewportInteraction refactoring; miscellaneous VREditor improvements

	- This merge contains various changes to ViewportInteraction and VR Editor to make things more extensible
	       - Some changes are from Yannick's recent shelved refactoring work, that I've modified and improved upon
	       - EditorWorldManager is gone and replaced by EditorWorldExtensionManager
	       - EditorWorldExtensions are sort of like a more modern version of 'FEdMode', but still a work in progress
	       - Cleaned up access to VREditorMode from other modules, forcing systems to go through EditorWorldExtensions

	- Overhauled how transforming objects works with world interaction
	       - Viewport interactors can now be used to move objects other than actors around (by implementing an UViewportTransformer, and a FViewportTransformable)
	       - Undo/Redo now works better with inertial transformation!  The transaction only ends when objects finally come to rest.
	       - Some initial support for 'grabber sphere' interactor methods has been implemented (not used yet)
	       - Viewport interaction input events now receive the viewport being interacted through
	       - Viewport interaction hover events no longer get a viewport client (because they must be designed to work any number of viewports.)

	Other changes:
	- The active Viewport Interaction 'gizmo mode' is now tied directly to the editor's normal gizmo mode.  They share the same state.
	- New console variable 'VI.UseTransientActors' can be turned off to force editor actors to be created non-transient to make it easier to debug
	- New console variable 'VI.DragTranslationVelocityStopEpsilon' sets the speed at which transformables will stop have inertia applied to them
	- Fixed cyclic dependencies with ViewportInteraction and VREditor modules not being tagged properly for UBT
	- Fixed some issues with transform gizmos not getting release events
	- Various methods that should have been const were made const
	- Eliminated duplicate implementation of SpawnTransientSceneActor and DestroyTransientActor; made it static
	- UnrealEd no longer directly depends on VREditor and ViewportInteraction modules
	- Engine: AActor::SetIsTemporarilyHiddenInEditor no longer does any work if the object's hidden state didn't change
	- Slate: Input preprocessor (if bound) now also gets a chance at mouse button down and up events

Change 3245645 on 2017/01/03 by [REDACTED_PERSON]

	Changing window closing to tab closing when Sequencer is invoked in VR Editor

Change 3246064 on 2017/01/04 by [REDACTED_PERSON]

	ViewportInteraction: Fixed static analysis warning about dereferencing (unused) null ptr

Change 3246179 on 2017/01/04 by [REDACTED_PERSON]

	Fix for crash when getting the world when creating help menus on the motion controllers.

Change 3246214 on 2017/01/04 by [REDACTED_PERSON]

	ViewportInteraction: Fixed assert when trying to use editor extensions in Simulate mode
	- Editor world extensions could not be bound to certain types of worlds; removed this restriction

Change 3246628 on 2017/01/04 by [REDACTED_PERSON]

	Making sequencer mesh trails non-selectable, so that only keys are transformable.

Change 3246837 on 2017/01/04 by [REDACTED_PERSON]

	VREditor: - Transform gizmo quarter rotation handles aiming towards the users view.
	- Transform gizmo added full rotation handle when dragging a rotation.
	- Removed bAllHandlesVisible arguments from UGizmoHandleGroup::UpdateGizmoHandleGroup.

Change 3246909 on 2017/01/04 by [REDACTED_PERSON]

	Mesh trail root component visibility was incorrectly propogating to children

Change 3246914 on 2017/01/04 by [REDACTED_PERSON]

	Fixing another GetWorld() call to go through the VRMode on the VR motion controller interactor

Change 3247755 on 2017/01/05 by [REDACTED_PERSON]

	Adjusting resolution of Sequencer window so keys are in the correct position

Change 3247829 on 2017/01/05 by [REDACTED_PERSON]

	Changing ViewportWorldInteraction's IsSelectableComponent check to respect bSelectable flag on primitive components

Change 3248447 on 2017/01/05 by [REDACTED_PERSON]

	Moving objects in VR with gizmo handles now auto-keys correctly.

Change 3248449 on 2017/01/05 by [REDACTED_PERSON]

	Disabling Sequencer fading and color scaling in VR Editor. Refreshing mesh trails when track data changes.

Change 3248571 on 2017/01/05 by [REDACTED_PERSON]

	Removing unused VREditorWorldInteraction::IsInteractableComponent and changing to a cast-based check for editor UI.

Change 3248623 on 2017/01/05 by [REDACTED_PERSON]

	VREditor: Fixed bad collision on 'quarter' transform gizmo handle
	- Mesh had bad simple collision which made it difficult to select (or impossible in some cases)

Change 3248639 on 2017/01/05 by [REDACTED_PERSON]

	ViewportInteraction: Made transform gizmo more compact when out of VR; minor other fixes [INTEGRATE]
	- Reduced size of transform gizmo when used in desktop mode (ForceVRMode or other modes that use this)
	- Reduced hover-scaled sized of new translation gizmo handles -- they were way too big!
	- Fixed rotation handles not taking into account animation transition alpha when scaling

Change 3249141 on 2017/01/06 by [REDACTED_PERSON]

	ViewportInteraction: You can no longer "get close to" the transform gizmo [INTEGRATE]
	- Distance-based scaling is now in effect even when closer than 1m to the gizmo origin

Change 3249143 on 2017/01/06 by [REDACTED_PERSON]

	Moving LevelVC settings inside check for LevelVC being valid.

Change 3249203 on 2017/01/06 by [REDACTED_PERSON]

	ViewportInteraction: Support hiding rotate/scale gizmo handles for single selected unoriented points [INTEGRATE]
	- Unoriented points in space can't be scaled or rotated unless more than one is selected, so it doesn't make sense to show those handles (plus they get in way.)
	- This will be used for a future feature

Change 3249421 on 2017/01/06 by [REDACTED_PERSON]

	VREditor: Fix transformgizmo rotation handle visibility when dragging freely and showing up while the current gizmo type is not rotation

Change 3249582 on 2017/01/06 by [REDACTED_PERSON]

	Removing VREditorDemo flag

Change 3255690 on 2017/01/12 by [REDACTED_PERSON]

	Adding alignment guidelines to transforming actors in the viewport. Use "EditorViewportSnapping.EnableGuides 1" to turn on.

Change 3342993 on 2017/03/13 by [REDACTED_PERSON]

	- Added CanPlaySound to UEditorEngine for external use.
	- Use bEnableEditorSounds setting for VREditor and ViewportWorldInteraction.

Change 3344798 on 2017/03/14 by [REDACTED_PERSON]

	- Added OnSwitchBeginPIEAndSIE event.
	- Disabled SIE to PIE for Oculus.
	- Moved bIsActive / bActive from VREditorMode and ViewportWorldInteraction to base class EditorWorldExtension.
	- Made bActive private in EditorWorldExtension.h.

Change 3344804 on 2017/03/14 by [REDACTED_PERSON]

	- Use optional member to store gizmo handle type for Universal gizmo.

Change 3347100 on 2017/03/15 by [REDACTED_PERSON]

	VREditor: Pop-up dialog the first time you enter VR mode or toggle auto-entry on

Change 3347452 on 2017/03/15 by [REDACTED_PERSON]

	Adjusting VREditor analytics events, including adding heartbeat

Change 3350214 on 2017/03/16 by [REDACTED_PERSON]

	Renaming r.AllowMotionBlur to vr.AllowMotionBlur

Change 3351473 on 2017/03/17 by [REDACTED_PERSON]

	VR Editor: Disable entering VR mode while in simulate.

Change 3351510 on 2017/03/17 by [REDACTED_PERSON]

	VR Editor: Fix Foliage Paint mesh does not render in desktop after painting in VREditor.

Change 3351787 on 2017/03/17 by [REDACTED_PERSON]

	VR Editor: - Gizmo animation when aiming at it only check for interactor that has a laser.
	- Fix rotation handle animation when selection changed.

diff --git a/Engine/Config/BaseEditorPerProjectUserSettings.ini b/Engine/Config/BaseEditorPerProjectUserSettings.ini
--- a/Engine/Config/BaseEditorPerProjectUserSettings.ini
+++ b/Engine/Config/BaseEditorPerProjectUserSettings.ini
@@ -154,12 +154,6 @@ bBreakOnExceptions=False
 bToolbarCustomization=False
 bBehaviorTreeEditor=False
 bEnableFindAndReplaceReferences =False
-; Enables a preview of the Unreal Editor in VR.  This adds a new tool bar button that allows you to toggle into "VR Mode" instantly.  This feature is still in development, but your feedback is appreciated!
-bEnableVREditing=false
-; If true, wearing a Vive or Oculus Rift headset will automatically enter VR Editing mode if Enable VR Editing is true. 
-bEnableAutoVREditMode=true
-; If true, the VR Tutorial window will always be shown at startup
-bAlwaysShowVRTutorial=true
 
 [/Script/UnrealEd.EditorLoadingSavingSettings]
 ; True if we should automatically load a default level at start up
diff --git a/Engine/Plugins/MovieScene/LevelSequenceEditor/Source/LevelSequenceEditor/Private/LevelSequenceEditorToolkit.cpp b/Engine/Plugins/MovieScene/LevelSequenceEditor/Source/LevelSequenceEditor/Private/LevelSequenceEditorToolkit.cpp
--- a/Engine/Plugins/MovieScene/LevelSequenceEditor/Source/LevelSequenceEditor/Private/LevelSequenceEditorToolkit.cpp
+++ b/Engine/Plugins/MovieScene/LevelSequenceEditor/Source/LevelSequenceEditor/Private/LevelSequenceEditorToolkit.cpp
@@ -48,9 +48,15 @@
 #include "UnrealEdMisc.h"
 
 // @todo sequencer: hack: setting defaults for transform tracks
+
 #include "Sections/MovieScene3DTransformSection.h"
 #include "Tracks/MovieScene3DTransformTrack.h"
 
+// To override Sequencer editor behavior for VR Editor 
+#include "EditorWorldExtension.h"
+#include "VREditorMode.h"
+
+
 #define LOCTEXT_NAMESPACE "LevelSequenceEditor"
 
 
@@ -157,6 +163,7 @@ FLevelSequenceEditorToolkit::FLevelSequenceEditorToolkit(const TSharedRef<ISlate
 		FAssetEditorExtender::CreateRaw(this, &FLevelSequenceEditorToolkit::HandleMenuExtensibilityGetExtender));
 	SequencerExtenderHandle = SequencerModule.GetAddTrackMenuExtensibilityManager()->GetExtenderDelegates()[NewIndex].GetHandle();
 
+
 	OpenToolkits.Add(this);
 }
 
@@ -254,6 +261,22 @@ void FLevelSequenceEditorToolkit::Initialize(const EToolkitMode::Type Mode, cons
 	LevelSequenceEditorModule.OnMasterSequenceCreated().AddRaw(this, &FLevelSequenceEditorToolkit::HandleMasterSequenceCreated);
 
 	FLevelSequenceEditorToolkit::OnOpened().Broadcast(*this);
+
+	{
+		UWorld* World = CastChecked<UWorld>( GetLevelSequenceEditorPlaybackContext() );
+		UVREditorMode* VRMode = Cast<UVREditorMode>( GEditor->GetEditorWorldExtensionsManager()->GetEditorWorldExtensions( World )->FindExtension( UVREditorMode::StaticClass() ) );
+		if (VRMode != nullptr)
+		{
+			VRMode->OnVREditingModeExit_Handler.BindSP(this, &FLevelSequenceEditorToolkit::HandleVREditorModeExit);
+			USequencerSettings& SavedSequencerSettings = *Sequencer->GetSequencerSettings();
+			VRMode->SaveSequencerSettings(Sequencer->GetKeyAllEnabled(), Sequencer->GetAutoKeyMode(), SavedSequencerSettings);
+			// Override currently set autokey behavior to always autokey all
+			Sequencer->SetAutoKeyMode(EAutoKeyMode::KeyAll);
+			Sequencer->SetKeyAllEnabled(true);
+			// Tell the VR Editor mode that Sequencer has refreshed
+			VRMode->RefreshVREditorSequencer(Sequencer.Get());
+		}
+	}
 }
 
 
@@ -588,6 +611,17 @@ void FLevelSequenceEditorToolkit::HandleActorAddedToSequencer(AActor* Actor, con
 }
 
 
+void FLevelSequenceEditorToolkit::HandleVREditorModeExit()
+{
+	UWorld* World = Cast<UWorld>(GetLevelSequenceEditorPlaybackContext());
+	UVREditorMode* VRMode = CastChecked<UVREditorMode>( GEditor->GetEditorWorldExtensionsManager()->GetEditorWorldExtensions( World )->FindExtension( UVREditorMode::StaticClass() ) );
+
+	// Reset sequencer settings
+	Sequencer->SetAutoKeyMode(VRMode->GetSavedEditorState().AutoKeyMode);
+	Sequencer->SetKeyAllEnabled(VRMode->GetSavedEditorState().bKeyAllEnabled);
+	VRMode->OnVREditingModeExit_Handler.Unbind();
+}
+
 void FLevelSequenceEditorToolkit::HandleMapChanged(class UWorld* NewWorld, EMapChangeType MapChangeType)
 {
 	// @todo sequencer: We should only wipe/respawn puppets that are affected by the world that is being changed! (multi-UWorld support)
@@ -660,7 +694,6 @@ void FLevelSequenceEditorToolkit::AddShot(UMovieSceneCinematicShotTrack* ShotTra
 		{
 			CameraGuid = GetSequencer()->CreateBinding(*NewCamera, NewCamera->GetActorLabel());
 		}
-
 		NewCamera->SetActorLocation( GCurrentLevelEditingViewportClient->GetViewLocation(), false );
 		NewCamera->SetActorRotation( GCurrentLevelEditingViewportClient->GetViewRotation() );
 		//pNewCamera->CameraComponent->FieldOfView = ViewportClient->ViewFOV; //@todo set the focal length from this field of view
@@ -821,6 +854,13 @@ void FLevelSequenceEditorToolkit::HandleTrackMenuExtensionAddTrack(FMenuBuilder&
 
 bool FLevelSequenceEditorToolkit::OnRequestClose()
 {
+	UWorld* World = CastChecked<UWorld>(GetLevelSequenceEditorPlaybackContext());
+	UVREditorMode* VRMode = Cast<UVREditorMode>(GEditor->GetEditorWorldExtensionsManager()->GetEditorWorldExtensions(World)->FindExtension(UVREditorMode::StaticClass()));
+	if (VRMode != nullptr)
+	{
+		// Null out the VR Mode's sequencer pointer
+		VRMode->RefreshVREditorSequencer(nullptr);
+	}
 	OpenToolkits.Remove(this);
 	OnClosedEvent.Broadcast();
 	return true;
diff --git a/Engine/Plugins/Runtime/OculusRift/Source/OculusRift/Private/OculusRiftHMD.cpp b/Engine/Plugins/Runtime/OculusRift/Source/OculusRift/Private/OculusRiftHMD.cpp
--- a/Engine/Plugins/Runtime/OculusRift/Source/OculusRift/Private/OculusRiftHMD.cpp
+++ b/Engine/Plugins/Runtime/OculusRift/Source/OculusRift/Private/OculusRiftHMD.cpp
@@ -580,17 +580,17 @@ bool FOculusRiftHMD::OnStartGameFrame( FWorldContext& InWorldContext )
 		const OVR::Vector3f newRight(CurrentSettings->EyeRenderDesc[1].HmdToEyeOffset);
 		const OVR::Vector3f prevLeft(MasterSettings->EyeRenderDesc[0].HmdToEyeOffset);
 		const OVR::Vector3f prevRight(MasterSettings->EyeRenderDesc[1].HmdToEyeOffset);
-		if (newLeft != prevLeft || newRight != prevRight)
+		// for debugging purposes only: overriding IPD
+		if( CurrentSettings->Flags.bOverrideIPD )
 		{
-			const float newIAD = newRight.Distance(newLeft);
-			UE_LOG(LogHMD, Log, TEXT("IAD has changed, new IAD is %.4f meters"), newIAD);
+			check( CurrentSettings->InterpupillaryDistance >= 0 );
+			CurrentSettings->EyeRenderDesc[ 0 ].HmdToEyeOffset.x = -CurrentSettings->InterpupillaryDistance * 0.5f;
+			CurrentSettings->EyeRenderDesc[ 1 ].HmdToEyeOffset.x = CurrentSettings->InterpupillaryDistance * 0.5f;
 		}
-		// for debugging purposes only: overriding IPD
-		if (CurrentSettings->Flags.bOverrideIPD)
+		else if (newLeft != prevLeft || newRight != prevRight)
 		{
-			check(CurrentSettings->InterpupillaryDistance >= 0);
-			CurrentSettings->EyeRenderDesc[0].HmdToEyeOffset.x = -CurrentSettings->InterpupillaryDistance * 0.5f;
-			CurrentSettings->EyeRenderDesc[1].HmdToEyeOffset.x = CurrentSettings->InterpupillaryDistance * 0.5f;
+			const float newIAD = newRight.Distance(newLeft);
+			UE_LOG(LogHMD, Log, TEXT("IAD has changed, new IAD is %.4f meters"), newIAD);
 		}
 #endif // #if !UE_BUILD_SHIPPING
 		// Save EyeRenderDesc in main settings.
@@ -1752,7 +1752,9 @@ void FOculusRiftHMD::SetupViewFamily(FSceneViewFamily& InViewFamily)
 	auto frame = GetFrame();
 	check(frame);
 
-	InViewFamily.EngineShowFlags.MotionBlur = 0;
+	extern RENDERER_API TAutoConsoleVariable<int32> CVarAllowMotionBlurInVR;
+
+	InViewFamily.EngineShowFlags.MotionBlur = CVarAllowMotionBlurInVR->GetInt() != 0;
 	InViewFamily.EngineShowFlags.HMDDistortion = false;
 	InViewFamily.EngineShowFlags.StereoRendering = IsStereoEnabled();
 	if (frame->Settings->Flags.bPauseRendering)
@@ -2242,7 +2244,7 @@ void FOculusRiftHMD::UpdateStereoRenderingParams()
 
 		for(int eyeIndex = 0; eyeIndex < ovrEye_Count; eyeIndex++)
 		{
-			vpSize[eyeIndex] = ovr_GetFovTextureSize(OvrSession, (ovrEyeType) eyeIndex, CurrentSettings->EyeFov[eyeIndex], pixelDensity);
+			vpSize[eyeIndex] = ovr_GetFovTextureSize(OvrSession, (ovrEyeType)eyeIndex, CurrentSettings->EyeFov[eyeIndex], pixelDensity);
 			vpSize[eyeIndex].w = FMath::Max(vpSize[eyeIndex].w, 1);
 			vpSize[eyeIndex].h = FMath::Max(vpSize[eyeIndex].h, 1);
 
@@ -2512,32 +2514,35 @@ void FOculusRiftHMD::OnBeginPlay(FWorldContext& InWorldContext)
 	CachedViewportWidget.Reset();
 	CachedWindow.Reset();
 
-#if WITH_EDITOR
-	// @TODO: add more values here.
-	// This call make sense when 'Play' is used from the Editor;
-	if (GIsEditor)
+	if( !GEnableVREditorHacks )
 	{
-		Settings->PositionOffset = FVector::ZeroVector;
-		Settings->BaseOrientation = FQuat::Identity;
-		Settings->BaseOffset = FVector::ZeroVector;
-		Settings->WorldToMetersScale = InWorldContext.World()->GetWorldSettings()->WorldToMeters;
-		//Settings->Flags.bWorldToMetersOverride = false;
-		InitDevice();
+#if WITH_EDITOR
+		// @TODO: add more values here.
+		// This call make sense when 'Play' is used from the Editor;
+		if( GIsEditor )
+		{
+			Settings->PositionOffset = FVector::ZeroVector;
+			Settings->BaseOrientation = FQuat::Identity;
+			Settings->BaseOffset = FVector::ZeroVector;
+			Settings->WorldToMetersScale = InWorldContext.World()->GetWorldSettings()->WorldToMeters;
+			//Settings->Flags.bWorldToMetersOverride = false;
+			InitDevice();
 
-		FApp::SetUseVRFocus(true);
-		FApp::SetHasVRFocus(true);
-		OnStartGameFrame(InWorldContext);
-	}
-	else
+			FApp::SetUseVRFocus( true );
+			FApp::SetHasVRFocus( true );
+			OnStartGameFrame( InWorldContext );
+		}
+		else
 #endif
-	{
-		MakeSureValidFrameExists(InWorldContext.World()->GetWorldSettings());
+		{
+			MakeSureValidFrameExists( InWorldContext.World()->GetWorldSettings() );
+		}
 	}
 }
 
 void FOculusRiftHMD::OnEndPlay(FWorldContext& InWorldContext)
 {
-	if (GIsEditor)
+	if (GIsEditor && !GEnableVREditorHacks)
 	{
 		// @todo vreditor: If we add support for starting PIE while in VR Editor, we don't want to kill stereo mode when exiting PIE
 		EnableStereo(false);
diff --git a/Engine/Plugins/Runtime/OculusRift/Source/OculusRift/Private/OculusRiftRender.cpp b/Engine/Plugins/Runtime/OculusRift/Source/OculusRift/Private/OculusRiftRender.cpp
--- a/Engine/Plugins/Runtime/OculusRift/Source/OculusRift/Private/OculusRiftRender.cpp
+++ b/Engine/Plugins/Runtime/OculusRift/Source/OculusRift/Private/OculusRiftRender.cpp
@@ -157,11 +157,6 @@ void FViewExtension::PreRenderView_RenderThread(FRHICommandListImmediate& RHICmd
 
 		CurrentFrame->PoseToOrientationAndPosition(CurrentFrame->EyeRenderPose[eyeIdx], GameEyeOrient, GameEyePosition);
 		const FQuat DeltaControlOrientation =  ViewOrientation * GameEyeOrient.Inverse();
-		// make sure we use the same viewrotation as we had on a game thread
-		if( !GEnableVREditorHacks )	// @todo vreditor: This assert goes off sometimes when dragging properties in the details pane (UE-27540)
-		{
-			check(View.ViewRotation == CurrentFrame->CachedViewRotation[eyeIdx]);
-		}
 
 		if (CurrentFrame->Flags.bOrientationChanged)
 		{
@@ -284,7 +279,15 @@ void FCustomPresent::CopyTexture_RenderThread(FRHICommandListImmediate& RHICmdLi
 	static FGlobalBoundShaderState BoundShaderState;
 	SetGlobalBoundShaderState(RHICmdList, FeatureLevel, BoundShaderState, RendererModule->GetFilterVertexDeclaration().VertexDeclarationRHI, *VertexShader, *PixelShader);
 
-	PixelShader->SetParameters(RHICmdList, TStaticSamplerState<SF_Bilinear>::GetRHI(), SrcTextureRHI);
+	const bool bSameSize = DstRect.Size() == SrcRect.Size();
+	if( bSameSize )
+	{
+		PixelShader->SetParameters(RHICmdList, TStaticSamplerState<SF_Point>::GetRHI(), SrcTextureRHI);
+	}
+	else
+	{
+		PixelShader->SetParameters(RHICmdList, TStaticSamplerState<SF_Bilinear>::GetRHI(), SrcTextureRHI);
+	}
 
 	RendererModule->DrawRectangle(
 		RHICmdList,
@@ -346,7 +349,7 @@ void FOculusRiftHMD::RenderTexture_RenderThread(class FRHICommandListImmediate&
 				const float SrcRectAspect = (float)SrcViewRect.Width() / (float)SrcViewRect.Height();
 				const float DstRectAspect = (float)DstViewRect.Width() / (float)DstViewRect.Height();
 
-				if( SrcRectAspect < 1.0f )
+				if( SrcRectAspect < DstRectAspect )
 				{
 					// Source is taller than destination
 					if( bCropToFill )
diff --git a/Engine/Plugins/Runtime/Steam/SteamVR/Source/SteamVR/Private/SteamVRHMD.cpp b/Engine/Plugins/Runtime/Steam/SteamVR/Source/SteamVR/Private/SteamVRHMD.cpp
--- a/Engine/Plugins/Runtime/Steam/SteamVR/Source/SteamVR/Private/SteamVRHMD.cpp
+++ b/Engine/Plugins/Runtime/Steam/SteamVR/Source/SteamVR/Private/SteamVRHMD.cpp
@@ -833,7 +833,10 @@ bool FSteamVRHMD::IsInLowPersistenceMode() const
 
 void FSteamVRHMD::OnEndPlay(FWorldContext& InWorldContext)
 {
-	EnableStereo(false);
+	if (!GEnableVREditorHacks)
+	{
+		EnableStereo(false);
+	}
 }
 
 bool FSteamVRHMD::OnStartGameFrame(FWorldContext& WorldContext)
@@ -1031,6 +1034,11 @@ bool FSteamVRHMD::IsStereoEnabled() const
 
 bool FSteamVRHMD::EnableStereo(bool bStereo)
 {
+	if( bStereoEnabled == bStereo )
+	{
+		return false;
+	}
+
 	if (bStereo && bIsQuitting)
 	{
 		// Cancel shutting down the vr subsystem if re-enabling stereo before we're done counting down
diff --git a/Engine/Source/Editor/ContentBrowser/Private/SAssetView.cpp b/Engine/Source/Editor/ContentBrowser/Private/SAssetView.cpp
--- a/Engine/Source/Editor/ContentBrowser/Private/SAssetView.cpp
+++ b/Engine/Source/Editor/ContentBrowser/Private/SAssetView.cpp
@@ -159,12 +159,6 @@ void SAssetView::Construct( const FArguments& InArgs )
 	TileViewThumbnailSize = 128;
 	TileViewThumbnailPadding = 5;
 
-	const bool bIsVREditorDemo = FParse::Param( FCommandLine::Get(), TEXT( "VREditorDemo" ) );	// @todo vreditor: Remove this when no longer needed
-	if( bIsVREditorDemo )
-	{
-		TileViewThumbnailPadding = 0;
-	}
-
 	TileViewNameHeight = 36;
 	ThumbnailScaleSliderValue = InArgs._ThumbnailScale; 
 
@@ -1051,12 +1045,6 @@ void SAssetView::CalculateFillScale( const FGeometry& AllottedGeometry )
 	{
 		float ItemWidth = GetTileViewItemBaseWidth();
 
- 		const bool bIsVREditorDemo = FParse::Param( FCommandLine::Get(), TEXT( "VREditorDemo" ) );	// @todo vreditor: Remove this when no longer needed
-		if( bIsVREditorDemo )
-		{
-			ItemWidth /= AllottedGeometry.Scale;
-		}
-
 		// Scrollbars are 16, but we add 1 to deal with half pixels.
 		const float ScrollbarWidth = 16 + 1;
 		float TotalWidth = AllottedGeometry.Size.X - ( ScrollbarWidth / AllottedGeometry.Scale );
@@ -3725,8 +3713,12 @@ FReply SAssetView::OnDraggingAssetItem( const FGeometry& MyGeometry, const FPoin
 			if( InAssetData.Num() > 0 )
 			{
 				UActorFactory* FactoryToUse = nullptr;
-				FEditorDelegates::OnAssetDragStarted.Broadcast( InAssetData, FactoryToUse );
-				if( MouseEvent.IsMouseButtonDown( EKeys::LeftMouseButton ) )
+				if( FEditorDelegates::OnAssetDragStarted.IsBound() )
+				{
+					FEditorDelegates::OnAssetDragStarted.Broadcast( InAssetData, FactoryToUse );
+					return FReply::Handled();
+				}
+				else if( MouseEvent.IsMouseButtonDown( EKeys::LeftMouseButton ) )
 				{
 					return FReply::Handled().BeginDragDrop( FAssetDragDropOp::New( InAssetData ) );
 				}
diff --git a/Engine/Source/Editor/DetailCustomizations/Private/FbxImportUIDetails.cpp b/Engine/Source/Editor/DetailCustomizations/Private/FbxImportUIDetails.cpp
--- a/Engine/Source/Editor/DetailCustomizations/Private/FbxImportUIDetails.cpp
+++ b/Engine/Source/Editor/DetailCustomizations/Private/FbxImportUIDetails.cpp
@@ -133,11 +133,6 @@ void FFbxImportUIDetails::CustomizeDetails( IDetailLayoutBuilder& DetailBuilder
 			Prop->SetOnPropertyValueChanged(FSimpleDelegate::CreateSP(this, &FFbxImportUIDetails::MeshImportModeChanged));
 			MeshCategory.AddProperty(Prop);
 		}
-		{
-			TSharedRef<IPropertyHandle> Prop = DetailBuilder.GetProperty(GET_MEMBER_NAME_CHECKED(UFbxImportUI, bImportAsSubDSurface));
-			Prop->SetOnPropertyValueChanged(FSimpleDelegate::CreateSP(this, &FFbxImportUIDetails::MeshImportModeChanged));
-			MeshCategory.AddProperty(Prop);
-		}
 	}
 
 	TSharedRef<IPropertyHandle> ImportMeshProp = DetailBuilder.GetProperty(GET_MEMBER_NAME_CHECKED(UFbxImportUI, bImportMesh));
diff --git a/Engine/Source/Editor/DetailCustomizations/Private/FbxImportUIDetails.h b/Engine/Source/Editor/DetailCustomizations/Private/FbxImportUIDetails.h
--- a/Engine/Source/Editor/DetailCustomizations/Private/FbxImportUIDetails.h
+++ b/Engine/Source/Editor/DetailCustomizations/Private/FbxImportUIDetails.h
@@ -36,7 +36,7 @@ public:
 	/** Called if the bImportMaterials changes */
 	void ImportMaterialsChanged();
 
-	/** Called if the mesh mode (static / skeletal / SubDSurface) changes */
+	/** Called if the mesh mode (static / skeletal) changes */
 	void MeshImportModeChanged();
 
 	/** Called if the import mesh option for skeletal meshes is changed */
diff --git a/Engine/Source/Editor/EditorStyle/Private/SlateEditorStyle.cpp b/Engine/Source/Editor/EditorStyle/Private/SlateEditorStyle.cpp
--- a/Engine/Source/Editor/EditorStyle/Private/SlateEditorStyle.cpp
+++ b/Engine/Source/Editor/EditorStyle/Private/SlateEditorStyle.cpp
@@ -4629,10 +4629,10 @@ void FSlateEditorStyle::FStyle::SetupLevelEditorStyle()
 		Set( "LevelEditor.FoliageMode.Selected", new IMAGE_BRUSH( "Icons/icon_Mode_Foliage_selected_40x", Icon40x40 ) );
 		Set( "LevelEditor.FoliageMode.Selected.Small", new IMAGE_BRUSH( "Icons/icon_Mode_Foliage_selected_40x", Icon20x20 ) );
 
-		Set( "LevelEditor.BspMode", new IMAGE_BRUSH( "Icons/icon_Mode_GeoEdit_40px", Icon40x40 ) );
-		Set( "LevelEditor.BspMode.Small", new IMAGE_BRUSH( "Icons/icon_Mode_GeoEdit_40px", Icon20x20 ) );
-		Set( "LevelEditor.BspMode.Selected", new IMAGE_BRUSH( "Icons/icon_Mode_GeoEdit-a_40px", Icon40x40 ) );
-		Set( "LevelEditor.BspMode.Selected.Small", new IMAGE_BRUSH( "Icons/icon_Mode_GeoEdit-a_40px", Icon20x20 ) );
+		Set("LevelEditor.BspMode", new IMAGE_BRUSH("Icons/icon_Mode_GeoEdit_40px", Icon40x40));
+		Set("LevelEditor.BspMode.Small", new IMAGE_BRUSH("Icons/icon_Mode_GeoEdit_40px", Icon20x20));
+		Set("LevelEditor.BspMode.Selected", new IMAGE_BRUSH("Icons/icon_Mode_GeoEdit-a_40px", Icon40x40));
+		Set("LevelEditor.BspMode.Selected.Small", new IMAGE_BRUSH("Icons/icon_Mode_GeoEdit-a_40px", Icon20x20));
 
 		Set( "LevelEditor.WorldProperties", new IMAGE_BRUSH( "Icons/icon_worldscript_40x", Icon40x40 ) );
 		Set( "LevelEditor.WorldProperties.Small", new IMAGE_BRUSH( "Icons/icon_worldscript_40x", Icon20x20 ) );
@@ -4649,8 +4649,8 @@ void FSlateEditorStyle::FStyle::SetupLevelEditorStyle()
 		Set( "LevelEditor.EditMatinee", new IMAGE_BRUSH( "Icons/icon_matinee_40x", Icon40x40 ) );
 		Set( "LevelEditor.EditMatinee.Small", new IMAGE_BRUSH( "Icons/icon_matinee_40x", Icon20x20 ) );
 
-		Set( "LevelEditor.ToggleVR", new IMAGE_BRUSH( "Icons/icon_ToggleVREditor_40x", Icon40x40 ) );
-		Set( "LevelEditor.ToggleVR.Small", new IMAGE_BRUSH( "Icons/icon_ToggleVREditor_40x", Icon20x20 ) );
+		Set( "LevelEditor.ToggleVR", new IMAGE_BRUSH( "Icons/VREditor/VR_Editor_Toolbar_Icon", Icon40x40 ) );
+		Set( "LevelEditor.ToggleVR.Small", new IMAGE_BRUSH( "Icons/VREditor/VR_Editor_Toolbar_Icon_Small", Icon20x20 ) );
 
 		Set( "MergeActors.MeshMergingTool", new IMAGE_BRUSH( "Icons/icon_MergeActors_MeshMerging_40x", Icon40x40 ) );
 		Set( "MergeActors.MeshProxyTool", new IMAGE_BRUSH( "Icons/icon_MergeActors_MeshProxy_40x", Icon40x40 ) );
@@ -5919,17 +5919,9 @@ void FSlateEditorStyle::FStyle::SetupContentBrowserStyle()
 		Set( "ContentBrowser.SourceTitleFont", TTF_CORE_FONT( "Fonts/Roboto-Regular", 12 ) );
 
 		// @todo vreditor urgent: Increase the size of Content Browser fonts while in VR
-		const bool bIsVREditorDemo = FParse::Param( FCommandLine::Get(), TEXT( "VREditorDemo" ) );	// @todo vreditor: Remove this when no longer needed
-		if( bIsVREditorDemo )
-		{
-			Set( "ContentBrowser.SourceListItemFont", TTF_CORE_FONT( "Fonts/Roboto-Regular", 20 ) );
-			Set( "ContentBrowser.SourceTreeItemFont", TTF_CORE_FONT( "Fonts/Roboto-Regular", 20 ) );
-		}
-		else
-		{
-			Set( "ContentBrowser.SourceListItemFont", TTF_CORE_FONT( "Fonts/Roboto-Regular", 10 ) );
-			Set( "ContentBrowser.SourceTreeItemFont", TTF_CORE_FONT( "Fonts/Roboto-Regular", 10 ) );
-		}
+		Set("ContentBrowser.SourceListItemFont", TTF_CORE_FONT("Fonts/Roboto-Regular", 10));
+		Set("ContentBrowser.SourceTreeItemFont", TTF_CORE_FONT("Fonts/Roboto-Regular", 10));
+
 		Set( "ContentBrowser.SourceTreeRootItemFont", TTF_CORE_FONT( "Fonts/Roboto-Regular", 12 ) );
 		Set( "ContentBrowser.AssetTreeFolderClosed", new IMAGE_BRUSH( "Icons/FolderClosed", FVector2D(18, 16) ) );
 		Set( "ContentBrowser.BreadcrumbPathPickerFolder", new IMAGE_BRUSH( "Icons/FolderClosed", FVector2D(18, 16), FLinearColor::Gray ) );
diff --git a/Engine/Source/Editor/FoliageEdit/Private/FoliageEdMode.cpp b/Engine/Source/Editor/FoliageEdit/Private/FoliageEdMode.cpp
--- a/Engine/Source/Editor/FoliageEdit/Private/FoliageEdMode.cpp
+++ b/Engine/Source/Editor/FoliageEdit/Private/FoliageEdMode.cpp
@@ -14,6 +14,7 @@
 #include "Framework/Application/SlateApplication.h"
 #include "Framework/Commands/UICommandList.h"
 #include "Materials/Material.h"
+#include "Materials/MaterialInstanceDynamic.h"
 #include "Engine/CollisionProfile.h"
 #include "Engine/StaticMeshActor.h"
 #include "Components/ModelComponent.h"
@@ -42,11 +43,10 @@
 #include "Components/BrushComponent.h"
 
 // VR Editor
-#include "IVREditorModule.h"
+#include "EditorWorldExtension.h"
 #include "VREditorMode.h"
 #include "ViewportWorldInteraction.h"
 #include "VREditorInteractor.h"
-#include "EditorWorldManager.h"
 
 
 #define LOCTEXT_NAMESPACE "FoliageEdMode"
@@ -54,6 +54,11 @@
 
 DEFINE_LOG_CATEGORY_STATIC(LogFoliage, Log, Warning);
 
+namespace VREd
+{
+	static FAutoConsoleVariable FoliageOpacity(TEXT("VREd.FoliageOpacity"), 0.02f, TEXT("The foliage brush opacity."));
+}
+
 //
 // FFoliageMeshUIInfo
 //
@@ -190,11 +195,13 @@ FEdModeFoliage::FEdModeFoliage()
 	, FoliageInteractor(nullptr)
 {
 	// Load resources and construct brush component
-	UMaterial* BrushMaterial = nullptr;
 	UStaticMesh* StaticMesh = nullptr;
 	if (!IsRunningCommandlet())
 	{
-		BrushMaterial = LoadObject<UMaterial>(nullptr, TEXT("/Engine/EditorLandscapeResources/FoliageBrushSphereMaterial.FoliageBrushSphereMaterial"), nullptr, LOAD_None, nullptr);
+		UMaterial* BrushMaterial = LoadObject<UMaterial>(nullptr, TEXT("/Engine/EditorLandscapeResources/FoliageBrushSphereMaterial.FoliageBrushSphereMaterial"), nullptr, LOAD_None, nullptr);
+		BrushMID = UMaterialInstanceDynamic::Create(BrushMaterial, GetTransientPackage());
+		check(BrushMID != nullptr);
+
 		StaticMesh = LoadObject<UStaticMesh>(nullptr, TEXT("/Engine/EngineMeshes/Sphere.Sphere"), nullptr, LOAD_None, nullptr);
 	}
 
@@ -202,12 +209,16 @@ FEdModeFoliage::FEdModeFoliage()
 	SphereBrushComponent->SetCollisionProfileName(UCollisionProfile::NoCollision_ProfileName);
 	SphereBrushComponent->SetCollisionObjectType(ECC_WorldDynamic);
 	SphereBrushComponent->SetStaticMesh(StaticMesh);
-	SphereBrushComponent->OverrideMaterials.Add(BrushMaterial);
+	SphereBrushComponent->SetMaterial(0, BrushMID);
 	SphereBrushComponent->SetAbsolute(true, true, true);
 	SphereBrushComponent->CastShadow = false;
 
 	bBrushTraceValid = false;
 	BrushLocation = FVector::ZeroVector;
+	
+	// Get the default opacity from the material.
+	FName OpacityParamName("OpacityAmount");
+	BrushMID->GetScalarParameterValue(OpacityParamName, DefaultBrushOpacity);
 
 	FFoliageEditCommands::Register();
 	UICommandList = MakeShareable(new FUICommandList);
@@ -357,37 +368,42 @@ void FEdModeFoliage::Enter()
 	NotifyNewCurrentLevel();
 
 	// Register for VR input events
-	TSharedPtr<FEditorWorldWrapper> EditorWorldWrapper = GEditor->GetEditorWorldManager()->GetEditorWorldWrapper(GetWorld());
-	if (EditorWorldWrapper.IsValid())
+	UViewportWorldInteraction* ViewportWorldInteraction = Cast<UViewportWorldInteraction>( GEditor->GetEditorWorldExtensionsManager()->GetEditorWorldExtensions( GetWorld() )->FindExtension( UViewportWorldInteraction::StaticClass() ) );
+	if (ViewportWorldInteraction != nullptr)
 	{
-		UViewportWorldInteraction* WorldInteraction = EditorWorldWrapper->GetViewportWorldInteraction();
-		if (WorldInteraction != nullptr)
-		{
-			WorldInteraction->OnViewportInteractionInputAction().RemoveAll(this);
-			WorldInteraction->OnViewportInteractionInputAction().AddRaw(this, &FEdModeFoliage::OnVRAction);
+		ViewportWorldInteraction->OnViewportInteractionInputAction().RemoveAll(this);
+		ViewportWorldInteraction->OnViewportInteractionInputAction().AddRaw(this, &FEdModeFoliage::OnVRAction);
 
-			WorldInteraction->OnViewportInteractionHoverUpdate().RemoveAll(this);
-			WorldInteraction->OnViewportInteractionHoverUpdate().AddRaw(this, &FEdModeFoliage::OnVRHoverUpdate);
-		}
+		ViewportWorldInteraction->OnViewportInteractionHoverUpdate().RemoveAll(this);
+		ViewportWorldInteraction->OnViewportInteractionHoverUpdate().AddRaw(this, &FEdModeFoliage::OnVRHoverUpdate);
+
+		// Hide the VR transform gizmo while we're in foliage mode.  It sort of gets in the way of painting.
+		ViewportWorldInteraction->SetTransformGizmoVisible( false );
+
+		SetBrushOpacity(VREd::FoliageOpacity->GetFloat());
 	}
+
+	// Make sure the brush is visible.
+	SphereBrushComponent->SetVisibility(true);
 }
 
 /** FEdMode: Called when the mode is exited */
 void FEdModeFoliage::Exit()
 {
 	// Unregister VR mode from event handlers
-	if (IVREditorModule::IsAvailable())
 	{
-		TSharedPtr<FEditorWorldWrapper> EditorWorldWrapper = GEditor->GetEditorWorldManager()->GetEditorWorldWrapper(GetWorld(), false);
-		if (EditorWorldWrapper.IsValid())
+		UViewportWorldInteraction* ViewportWorldInteraction = Cast<UViewportWorldInteraction>( GEditor->GetEditorWorldExtensionsManager()->GetEditorWorldExtensions( GetWorld() )->FindExtension( UViewportWorldInteraction::StaticClass() ) );
+		if (ViewportWorldInteraction != nullptr)
 		{
-			UViewportWorldInteraction* WorldInteraction = EditorWorldWrapper->GetViewportWorldInteraction();
-			if (WorldInteraction != nullptr)
-			{
-				WorldInteraction->OnViewportInteractionInputAction().RemoveAll(this);
-				WorldInteraction->OnViewportInteractionHoverUpdate().RemoveAll(this);
-				FoliageInteractor = nullptr;
-			}
+			// Restore the transform gizmo visibility
+			ViewportWorldInteraction->SetTransformGizmoVisible( true );
+
+			ViewportWorldInteraction->OnViewportInteractionInputAction().RemoveAll(this);
+			ViewportWorldInteraction->OnViewportInteractionHoverUpdate().RemoveAll(this);
+			FoliageInteractor = nullptr;
+
+			// Reset the brush opacity to default.
+			SetBrushOpacity(DefaultBrushOpacity);
 		}
 	}
 
@@ -447,134 +463,136 @@ void FEdModeFoliage::Exit()
 	FEdMode::Exit();
 }
 
-void FEdModeFoliage::OnVRHoverUpdate(FEditorViewportClient& ViewportClient, UViewportInteractor* Interactor, FVector& HoverImpactPoint, bool& bWasHandled)
+void FEdModeFoliage::OnVRHoverUpdate(UViewportInteractor* Interactor, FVector& HoverImpactPoint, bool& bWasHandled)
 {
-	// Check if VR Editor is active
-	if (IVREditorModule::Get().IsVREditorModeActive())
+	UVREditorMode* VREditorMode = Cast<UVREditorMode>( GEditor->GetEditorWorldExtensionsManager()->GetEditorWorldExtensions( GetWorld() )->FindExtension( UVREditorMode::StaticClass() ) );
+	if( VREditorMode != nullptr && VREditorMode->IsFullyInitialized() )
 	{
-		UVREditorMode* VREditorMode = IVREditorModule::Get().GetVREditorMode();
-		if (VREditorMode != nullptr && VREditorMode->IsFullyInitialized())
+		// Check if we're hovering over UI. If so, stop painting so we don't display the preview brush sphere
+		if (FoliageInteractor && (FoliageInteractor->IsHoveringOverPriorityType() || FoliageInteractor->GetDraggingMode() != EViewportInteractionDraggingMode::Nothing))
 		{
-			// Check if we're hovering over UI. If so, stop painting so we don't display the preview brush sphere
-			if (FoliageInteractor && FoliageInteractor->IsHoveringOverPriorityType())
-			{
-				EndFoliageBrushTrace();
-				FoliageInteractor = nullptr;
-			}
-			// If there isn't currently a foliage interactor and we are hovering over something valid
-			else if (FoliageInteractor == nullptr && Interactor->GetHitResultFromLaserPointer().GetActor() != nullptr)
-			{
-				FoliageInteractor = Interactor;
-			}
-			// If we aren't hovering over something valid and the tool isn't active
-			else if (Interactor->GetHitResultFromLaserPointer().GetActor() == nullptr && !bToolActive)
-			{
-				FoliageInteractor = nullptr;
-			}
+			EndFoliageBrushTrace();
+			FoliageInteractor = nullptr;
+		}
+		// If there isn't currently a foliage interactor and we are hovering over something valid
+		else if (FoliageInteractor == nullptr && !Interactor->IsHoveringOverPriorityType() && Interactor->GetHitResultFromLaserPointer().GetActor() != nullptr)
+		{
+			FoliageInteractor = Interactor;
+		}
+		// If we aren't hovering over something valid and the tool isn't active
+		else if (Interactor->GetHitResultFromLaserPointer().GetActor() == nullptr && !bToolActive)
+		{
+			FoliageInteractor = nullptr;
+		}
 
-			// Skip other interactors if we are painting with one
-			if (FoliageInteractor && Interactor == FoliageInteractor)
+		// Skip other interactors if we are painting with one
+		if (FoliageInteractor && Interactor == FoliageInteractor)
+		{
+			// Go ahead and paint immediately
+			FVector LaserPointerStart, LaserPointerEnd;
+			if (FoliageInteractor->GetLaserPointer( /* Out */ LaserPointerStart, /* Out */ LaserPointerEnd))
 			{
-				// Go ahead and paint immediately
-				FVector LaserPointerStart, LaserPointerEnd;
-				if (FoliageInteractor->GetLaserPointer( /* Out */ LaserPointerStart, /* Out */ LaserPointerEnd))
-				{
-					const FVector LaserPointerDirection = (LaserPointerEnd - LaserPointerStart).GetSafeNormal();
+				const FVector LaserPointerDirection = (LaserPointerEnd - LaserPointerStart).GetSafeNormal();
 
-					FoliageBrushTrace(&ViewportClient, LaserPointerStart, LaserPointerDirection);
+				FoliageBrushTrace(nullptr, LaserPointerStart, LaserPointerDirection);
+
+				if( bBrushTraceValid )
+				{
+					HoverImpactPoint = BrushLocation;
+					bWasHandled = true;
 				}
 			}
 		}
+		const bool bBrushMeshVisible = !(FoliageInteractor == nullptr || Interactor->GetDraggingMode() != EViewportInteractionDraggingMode::Nothing);
+		SphereBrushComponent->SetVisibility(bBrushMeshVisible);
 	}
 }
 
 void FEdModeFoliage::OnVRAction(class FEditorViewportClient& ViewportClient, UViewportInteractor* Interactor, const FViewportActionKeyInput& Action, bool& bOutIsInputCaptured, bool& bWasHandled)
 {
-	if (IVREditorModule::Get().IsVREditorModeActive())
+	UVREditorMode* VREditorMode = Cast<UVREditorMode>( GEditor->GetEditorWorldExtensionsManager()->GetEditorWorldExtensions( GetWorld() )->FindExtension( UVREditorMode::StaticClass() ) );
+	if (VREditorMode != nullptr && Interactor != nullptr && Interactor->GetDraggingMode() == EViewportInteractionDraggingMode::Nothing)
 	{
-		UVREditorMode* VREditorMode = IVREditorModule::Get().GetVREditorMode();
-		if (VREditorMode != nullptr && Interactor != nullptr)
+		const UVREditorInteractor* VREditorInteractor = Cast<UVREditorInteractor>(Interactor);
+		if (Action.ActionType == ViewportWorldActionTypes::SelectAndMove && (VREditorInteractor == nullptr || !VREditorMode->IsShowingRadialMenu(VREditorInteractor)))
 		{
-			const UVREditorInteractor* VRInteractor = Cast<UVREditorInteractor>(Interactor);
-
-			// Only allow light press
-			Interactor->SetAllowTriggerFullPress(false);
-			// Consume light press
-			if ((Action.ActionType == ViewportWorldActionTypes::SelectAndMove_LightlyPressed) && (!VREditorMode->IsShowingRadialMenu(VRInteractor)))
+			if (Action.Event == IE_Pressed && !Interactor->IsHoveringOverPriorityType())
 			{
-				if ((Action.Event == IE_Pressed) && (!VRInteractor->IsHoveringOverPriorityType()))
+				bWasHandled = true;
+				bOutIsInputCaptured = true;
+
+				// Go ahead and paint immediately
+				FVector LaserPointerStart, LaserPointerEnd;
+				if (Interactor->GetLaserPointer( /* Out */ LaserPointerStart, /* Out */ LaserPointerEnd))
 				{
-					// Go ahead and paint immediately
-					FVector LaserPointerStart, LaserPointerEnd;
-					if (Interactor->GetLaserPointer( /* Out */ LaserPointerStart, /* Out */ LaserPointerEnd))
+					const FVector LaserPointerDirection = (LaserPointerEnd - LaserPointerStart).GetSafeNormal();
+					BrushTraceDirection = LaserPointerDirection;
+
+					// Only start painting if we're not dragging a widget handle
+					if (ViewportClient.GetCurrentWidgetAxis() == EAxisList::None)
 					{
-						const FVector LaserPointerDirection = (LaserPointerEnd - LaserPointerStart).GetSafeNormal();
-						BrushTraceDirection = LaserPointerDirection;
+						if (UISettings.GetPaintToolSelected() || UISettings.GetReapplyToolSelected() || UISettings.GetLassoSelectToolSelected())
+						{
+							StartFoliageBrushTrace(&ViewportClient, Interactor);
+							FoliageBrushTrace(&ViewportClient, LaserPointerStart, LaserPointerDirection);
+						}
 
-						// Only start painting if we're not dragging a widget handle
-						if (ViewportClient.GetCurrentWidgetAxis() == EAxisList::None)
+						// Fill a static mesh with foliage brush
+						else if (UISettings.GetPaintBucketToolSelected() || UISettings.GetReapplyPaintBucketToolSelected())
 						{
-							if (UISettings.GetPaintToolSelected() || UISettings.GetReapplyToolSelected() || UISettings.GetLassoSelectToolSelected())
-							{
-								StartFoliageBrushTrace(&ViewportClient, Interactor);
-								FoliageBrushTrace(&ViewportClient, LaserPointerStart, LaserPointerDirection);
-							}
+							FHitResult HitResult = Interactor->GetHitResultFromLaserPointer();
 
-							// Fill a static mesh with foliage brush
-							else if (UISettings.GetPaintBucketToolSelected() || UISettings.GetReapplyPaintBucketToolSelected())
+							if (HitResult.Actor.Get() != nullptr)
 							{
-								FHitResult HitResult = Interactor->GetHitResultFromLaserPointer();
+								GEditor->BeginTransaction(NSLOCTEXT("UnrealEd", "FoliageMode_EditTransaction", "Foliage Editing"));
 
-								if (HitResult.Actor.Get() != nullptr)
+								if (IsModifierButtonPressed(&ViewportClient))
 								{
-									GEditor->BeginTransaction(NSLOCTEXT("UnrealEd", "FoliageMode_EditTransaction", "Foliage Editing"));
-
-									if (IsModifierButtonPressed(&ViewportClient))
-									{
-										ApplyPaintBucket_Remove(HitResult.Actor.Get());
-									}
-									else
-									{
-										ApplyPaintBucket_Add(HitResult.Actor.Get());
-									}
-
-									GEditor->EndTransaction();
+									ApplyPaintBucket_Remove(HitResult.Actor.Get());
 								}
+								else
+								{
+									ApplyPaintBucket_Add(HitResult.Actor.Get());
+								}
+
+								GEditor->EndTransaction();
 							}
-							// Select an instanced foliage
-							else if (UISettings.GetSelectToolSelected())
-							{
-								FHitResult HitResult = Interactor->GetHitResultFromLaserPointer();
+						}
+						// Select an instanced foliage
+						else if (UISettings.GetSelectToolSelected())
+						{
+							FHitResult HitResult = Interactor->GetHitResultFromLaserPointer();
 
-								GEditor->BeginTransaction(NSLOCTEXT("UnrealEd", "FoliageMode_EditTransaction", "Foliage Editing"));
+							GEditor->BeginTransaction(NSLOCTEXT("UnrealEd", "FoliageMode_EditTransaction", "Foliage Editing"));
 
-								if (HitResult.GetActor() != nullptr)
+							if (HitResult.GetActor() != nullptr)
+							{
+								// Clear all currently selected instances
+								SelectInstances(ViewportClient.GetWorld(), false);
+								for (auto& FoliageMeshUI : FoliageMeshList)
 								{
-									// Clear all currently selected instances
-									SelectInstances(ViewportClient.GetWorld(), false);
-									for (auto& FoliageMeshUI : FoliageMeshList)
-									{
-										UFoliageType* Settings = FoliageMeshUI->Settings;
-										SelectInstanceAtLocation(ViewportClient.GetWorld(), Settings, HitResult.ImpactPoint, !IsModifierButtonPressed(&ViewportClient));
-									}
+									UFoliageType* Settings = FoliageMeshUI->Settings;
+									SelectInstanceAtLocation(ViewportClient.GetWorld(), Settings, HitResult.ImpactPoint, !IsModifierButtonPressed(&ViewportClient));
 								}
+							}
 
-								GEditor->EndTransaction();
+							GEditor->EndTransaction();
 
-								// @todo vreditor: we currently don't have a key mapping scheme to snap selected instances to ground 
-								// SnapSelectedInstancesToGround(GetWorld());
+							// @todo vreditor: we currently don't have a key mapping scheme to snap selected instances to ground 
+							// SnapSelectedInstancesToGround(GetWorld());
 
-							}
 						}
 					}
 				}
+			}
+			// Stop current tracking if the user is no longer painting
+			else if (Action.Event == IE_Released && FoliageInteractor && (FoliageInteractor == Interactor))
+			{
+				EndFoliageBrushTrace();
+				FoliageInteractor = nullptr;
 
-				// Stop current tracking if the user is no longer painting
-				else if (Action.Event == IE_Released && FoliageInteractor && (FoliageInteractor == Interactor))
-				{
-					EndFoliageBrushTrace();
-					FoliageInteractor = nullptr;
-				}
+				bWasHandled = true;
+				bOutIsInputCaptured = false;
 			}
 		}
 	}
@@ -794,7 +812,7 @@ void FEdModeFoliage::EndFoliageBrushTrace()
 void FEdModeFoliage::FoliageBrushTrace(FEditorViewportClient* ViewportClient, const FVector& InRayOrigin, const FVector& InRayDirection)
 {
 	bBrushTraceValid = false;
-	if (!ViewportClient->IsMovingCamera() && ViewportClient->IsVisible())
+	if (ViewportClient == nullptr || ( !ViewportClient->IsMovingCamera() && ViewportClient->IsVisible() ) )
 	{
 		if (UISettings.GetPaintToolSelected() || UISettings.GetReapplyToolSelected() || UISettings.GetLassoSelectToolSelected())
 		{
@@ -802,7 +820,7 @@ void FEdModeFoliage::FoliageBrushTrace(FEditorViewportClient* ViewportClient, co
 			const FVector TraceEnd(InRayOrigin + InRayDirection * HALF_WORLD_MAX);
 
 			FHitResult Hit;
-			UWorld* World = ViewportClient->GetWorld();
+			UWorld* World = GetWorld();
 			static FName NAME_FoliageBrush = FName(TEXT("FoliageBrush"));
 			FFoliagePaintingGeometryFilter FilterFunc = FFoliagePaintingGeometryFilter(UISettings);
 
@@ -839,7 +857,8 @@ void FEdModeFoliage::FoliageBrushTrace(FEditorViewportClient* ViewportClient, co
 bool FEdModeFoliage::MouseMove(FEditorViewportClient* ViewportClient, FViewport* Viewport, int32 MouseX, int32 MouseY)
 {
 	// Use mouse capture if there's no other interactor currently tracing brush
-	if (!IVREditorModule::Get().IsVREditorModeActive())
+	UVREditorMode* VREditorMode = Cast<UVREditorMode>( GEditor->GetEditorWorldExtensionsManager()->GetEditorWorldExtensions( GetWorld() )->FindExtension( UVREditorMode::StaticClass() ) );
+	if (VREditorMode == nullptr || !VREditorMode->IsActive())
 	{
 		// Compute a world space ray from the screen space mouse coordinates
 		FSceneViewFamilyContext ViewFamily(FSceneViewFamily::ConstructionValues(
@@ -876,7 +895,8 @@ bool FEdModeFoliage::MouseMove(FEditorViewportClient* ViewportClient, FViewport*
 bool FEdModeFoliage::CapturedMouseMove(FEditorViewportClient* ViewportClient, FViewport* Viewport, int32 MouseX, int32 MouseY)
 {
 	// Use mouse capture if there's no other interactor currently tracing brush
-	if (!IVREditorModule::Get().IsVREditorModeActive())
+	UVREditorMode* VREditorMode = Cast<UVREditorMode>( GEditor->GetEditorWorldExtensionsManager()->GetEditorWorldExtensions( GetWorld() )->FindExtension( UVREditorMode::StaticClass() ) );
+	if (VREditorMode == nullptr || !VREditorMode->IsActive())
 	{
 		//Compute a world space ray from the screen space mouse coordinates
 		FSceneViewFamilyContext ViewFamily(FSceneViewFamily::ConstructionValues(
@@ -1126,6 +1146,12 @@ bool FEdModeFoliage::VertexMaskCheck(const FHitResult& Hit, const UFoliageType*
 	return true;
 }
 
+void FEdModeFoliage::SetBrushOpacity(const float InOpacity)
+{
+	static FName OpacityParamName("OpacityAmount");
+	BrushMID->SetScalarParameterValue(OpacityParamName, InOpacity);
+}
+
 bool LandscapeLayerCheck(const FHitResult& Hit, const UFoliageType* Settings, FEdModeFoliage::LandscapeLayerCacheData& LandscapeLayersCache, float& OutHitWeight)
 {
 	OutHitWeight = 1.f;
@@ -3489,4 +3515,4 @@ void FFoliageUISettings::Save()
 	GConfig->SetFloat(TEXT("FoliageEdit"), TEXT("PaletteThumbnailScale"), PaletteThumbnailScale, GEditorPerProjectIni);
 }
 
-#undef LOCTEXT_NAMESPACE
+#undef LOCTEXT_NAMESPACE
\ No newline at end of file
diff --git a/Engine/Source/Editor/FoliageEdit/Private/FoliageEdMode.h b/Engine/Source/Editor/FoliageEdit/Private/FoliageEdMode.h
--- a/Engine/Source/Editor/FoliageEdit/Private/FoliageEdMode.h
+++ b/Engine/Source/Editor/FoliageEdit/Private/FoliageEdMode.h
@@ -116,7 +116,7 @@ struct FFoliageUISettings
 		, bShowPaletteItemTooltips(true)
 		, ActivePaletteViewMode(EFoliagePaletteViewMode::Thumbnail)
 		, PaletteThumbnailScale(0.3f)
-		, Radius(512.f)
+		, Radius(250.f)
 		, PaintDensity(0.5f)
 		, UnpaintDensity(0.f)
 		, IsInSingleInstantiationMode(false)
@@ -456,7 +456,7 @@ public:
 	void OnVRAction(class FEditorViewportClient& ViewportClient, class UViewportInteractor* Interactor, const struct FViewportActionKeyInput& Action, bool& bOutIsInputCaptured, bool& bWasHandled);
 
 	/** Called on VR hovering */
-	void OnVRHoverUpdate(FEditorViewportClient& ViewportClient, UViewportInteractor* Interactor, FVector& HoverImpactPoint, bool& bWasHandled);
+	void OnVRHoverUpdate(UViewportInteractor* Interactor, FVector& HoverImpactPoint, bool& bWasHandled);
 
 	typedef TMap<FName, TMap<ULandscapeComponent*, TArray<uint8> > > LandscapeLayerCacheData;
 
@@ -559,12 +559,21 @@ private:
 	/** Does a filter based on the vertex color of a static mesh */
 	static bool VertexMaskCheck(const FHitResult& Hit, const UFoliageType* Settings);
 
+	/** Set the brush mesh opacity */
+	void SetBrushOpacity(const float InOpacity);
+
 	bool bBrushTraceValid;
 	FVector BrushLocation;
 	FVector BrushNormal;
 	FVector BrushTraceDirection;
 	UStaticMeshComponent* SphereBrushComponent;
 
+	/** The dynamic material of the sphere brush. */
+	class UMaterialInstanceDynamic* BrushMID;
+
+	/** Default opacity received from the brush material to reset it when closing. */
+	float DefaultBrushOpacity;
+
 	// Landscape layer cache data
 	LandscapeLayerCacheData LandscapeLayerCaches;
 
diff --git a/Engine/Source/Editor/LandscapeEditor/Private/LandscapeEdMode.cpp b/Engine/Source/Editor/LandscapeEditor/Private/LandscapeEdMode.cpp
--- a/Engine/Source/Editor/LandscapeEditor/Private/LandscapeEdMode.cpp
+++ b/Engine/Source/Editor/LandscapeEditor/Private/LandscapeEdMode.cpp
@@ -26,7 +26,7 @@
 #include "Toolkits/ToolkitManager.h"
 #include "LandscapeHeightfieldCollisionComponent.h"
 #include "InstancedFoliageActor.h"
-#include "IVREditorModule.h"
+#include "EditorWorldExtension.h"
 #include "ViewportWorldInteraction.h"
 #include "VREditorInteractor.h"
 #include "LandscapeEdModeTools.h"
@@ -40,13 +40,13 @@
 
 // VR Editor
 #include "VREditorMode.h"
-#include "EditorWorldManager.h"
 
 // Classes
 #include "LandscapeMaterialInstanceConstant.h"
 #include "LandscapeSplinesComponent.h"
 #include "ComponentReregisterContext.h"
 #include "EngineUtils.h"
+#include "IVREditorModule.h"
 
 #define LOCTEXT_NAMESPACE "Landscape"
 
@@ -480,18 +480,16 @@ void FEdModeLandscape::Enter()
 	}
 
 	// Register to find out about VR input events
-	TSharedPtr<FEditorWorldWrapper> EditorWorldWrapper = GEditor->GetEditorWorldManager()->GetEditorWorldWrapper(GetWorld(), false);
-	if (EditorWorldWrapper.IsValid())
+	UViewportWorldInteraction* ViewportWorldInteraction = Cast<UViewportWorldInteraction>(GEditor->GetEditorWorldExtensionsManager()->GetEditorWorldExtensions(GetWorld())->FindExtension(UViewportWorldInteraction::StaticClass()));
+	if (ViewportWorldInteraction != nullptr)
 	{
-		UViewportWorldInteraction* WorldInteraction = EditorWorldWrapper->GetViewportWorldInteraction();
-		if (WorldInteraction != nullptr)
-		{
-			WorldInteraction->OnViewportInteractionInputAction().RemoveAll(this);
-			WorldInteraction->OnViewportInteractionInputAction().AddRaw(this, &FEdModeLandscape::OnVRAction);
 
-			WorldInteraction->OnViewportInteractionHoverUpdate().RemoveAll(this);
-			WorldInteraction->OnViewportInteractionHoverUpdate().AddRaw(this, &FEdModeLandscape::OnVRHoverUpdate);
-		}
+			ViewportWorldInteraction->OnViewportInteractionInputAction().RemoveAll(this);
+			ViewportWorldInteraction->OnViewportInteractionInputAction().AddRaw(this, &FEdModeLandscape::OnVRAction);
+
+			ViewportWorldInteraction->OnViewportInteractionHoverUpdate().RemoveAll(this);
+			ViewportWorldInteraction->OnViewportInteractionHoverUpdate().AddRaw(this, &FEdModeLandscape::OnVRHoverUpdate);
+
 	}
 }
 
@@ -500,18 +498,11 @@ void FEdModeLandscape::Enter()
 void FEdModeLandscape::Exit()
 {
 	// Unregister VR mode from event handlers
-	if (IVREditorModule::IsAvailable())
+	UViewportWorldInteraction* ViewportWorldInteraction = Cast<UViewportWorldInteraction>(GEditor->GetEditorWorldExtensionsManager()->GetEditorWorldExtensions(GetWorld())->FindExtension(UViewportWorldInteraction::StaticClass()));
+	if (ViewportWorldInteraction != nullptr)
 	{
-		TSharedPtr<FEditorWorldWrapper> EditorWorldWrapper = GEditor->GetEditorWorldManager()->GetEditorWorldWrapper(GetWorld(), false);
-		if (EditorWorldWrapper.IsValid())
-		{
-			UViewportWorldInteraction* WorldInteraction = EditorWorldWrapper->GetViewportWorldInteraction();
-			if (WorldInteraction != nullptr)
-			{
-				WorldInteraction->OnViewportInteractionInputAction().RemoveAll(this);
-				WorldInteraction->OnViewportInteractionHoverUpdate().RemoveAll(this);
-			}
-		}
+		ViewportWorldInteraction->OnViewportInteractionInputAction().RemoveAll(this);
+		ViewportWorldInteraction->OnViewportInteractionHoverUpdate().RemoveAll(this);
 	}
 	
 	FEditorSupportDelegates::WorldChange.Remove(OnWorldChangeDelegateHandle);
@@ -583,12 +574,12 @@ void FEdModeLandscape::Exit()
 }
 
 
-void FEdModeLandscape::OnVRHoverUpdate(FEditorViewportClient& ViewportClient, UViewportInteractor* Interactor, FVector& HoverImpactPoint, bool& bWasHandled)
+void FEdModeLandscape::OnVRHoverUpdate(UViewportInteractor* Interactor, FVector& HoverImpactPoint, bool& bWasHandled)
 {
 	if (InteractorPainting != nullptr && InteractorPainting == Interactor && IVREditorModule::Get().IsVREditorModeActive())
 	{
-		UVREditorMode* VREditorMode = IVREditorModule::Get().GetVREditorMode();
-		if (VREditorMode != nullptr && VREditorMode->IsActive() && Interactor != nullptr)
+		UVREditorMode* VREditorMode = Cast<UVREditorMode>( GEditor->GetEditorWorldExtensionsManager()->GetEditorWorldExtensions( GetWorld() )->FindExtension( UVREditorMode::StaticClass() ) );
+		if( VREditorMode != nullptr && VREditorMode->IsActive() && Interactor != nullptr && Interactor->GetDraggingMode() == EViewportInteractionDraggingMode::Nothing )
 		{
 			const UVREditorInteractor* VRInteractor = Cast<UVREditorInteractor>(Interactor);
 
@@ -598,16 +589,8 @@ void FEdModeLandscape::OnVRHoverUpdate(FEditorViewportClient& ViewportClient, UV
 				FVector LaserPointerStart, LaserPointerEnd;
 				if (Interactor->GetLaserPointer( /* Out */ LaserPointerStart, /* Out */ LaserPointerEnd))
 				{
-					if (LandscapeTrace(&ViewportClient, LaserPointerStart, LaserPointerEnd, HitLocation))
+					if( LandscapeTrace( LaserPointerStart, LaserPointerEnd, HitLocation ) )
 					{
-						if (bIsPaintingInVR && !(CurrentToolTarget.TargetType == ELandscapeToolTargetType::Weightmap && CurrentToolTarget.LayerInfo == NULL))
-						{
-							if (CurrentTool->BeginTool(&ViewportClient, CurrentToolTarget, HitLocation, Interactor))
-							{
-								ToolActiveViewport = ViewportClient.Viewport;
-							}
-						}
-
 						if (CurrentBrush)
 						{
 							// Inform the brush of the current location, to update the cursor
@@ -622,12 +605,13 @@ void FEdModeLandscape::OnVRHoverUpdate(FEditorViewportClient& ViewportClient, UV
 
 void FEdModeLandscape::OnVRAction(FEditorViewportClient& ViewportClient, UViewportInteractor* Interactor, const struct FViewportActionKeyInput& Action, bool& bOutIsInputCaptured, bool& bWasHandled)
 {
+	UVREditorMode* VREditorMode = Cast<UVREditorMode>( GEditor->GetEditorWorldExtensionsManager()->GetEditorWorldExtensions( GetWorld() )->FindExtension( UVREditorMode::StaticClass() ) );
 	// Never show the traditional Unreal transform widget.  It doesn't work in VR because we don't have hit proxies.
 	ViewportClient.EngineShowFlags.SetModeWidgets(false);
 
-	if (IVREditorModule::Get().IsVREditorModeActive() && Interactor != nullptr)
+	if (VREditorMode != nullptr && VREditorMode->IsActive() && Interactor != nullptr && Interactor->GetDraggingMode() == EViewportInteractionDraggingMode::Nothing)
 	{
-		if (Action.ActionType == ViewportWorldActionTypes::SelectAndMove_LightlyPressed || Action.ActionType == ViewportWorldActionTypes::SelectAndMove)
+		if (Action.ActionType == ViewportWorldActionTypes::SelectAndMove)
 		{
 			const UVREditorInteractor* VRInteractor = Cast<UVREditorInteractor>(Interactor);
 
@@ -646,7 +630,7 @@ void FEdModeLandscape::OnVRAction(FEditorViewportClient& ViewportClient, UViewpo
 					FVector LaserPointerStart, LaserPointerEnd;
 					if (Interactor->GetLaserPointer( /* Out */ LaserPointerStart, /* Out */ LaserPointerEnd))
 					{
-						if (LandscapeTrace(&ViewportClient, LaserPointerStart, LaserPointerEnd, HitLocation))
+						if (LandscapeTrace(LaserPointerStart, LaserPointerEnd, HitLocation))
 						{
 							if (!(CurrentToolTarget.TargetType == ELandscapeToolTargetType::Weightmap && CurrentToolTarget.LayerInfo == NULL))
 							{
@@ -926,16 +910,16 @@ bool FEdModeLandscape::LandscapeMouseTrace(FEditorViewportClient* ViewportClient
 
 	FVector End = Start + WORLD_MAX * BrushTraceDirection;
 
-	return LandscapeTrace(ViewportClient, Start, End, OutHitLocation);
+	return LandscapeTrace(Start, End, OutHitLocation);
 }
 
-bool FEdModeLandscape::LandscapeTrace(const FEditorViewportClient* ViewportClient, const FVector& InRayOrigin, const FVector& InRayEnd, FVector& OutHitLocation)
+bool FEdModeLandscape::LandscapeTrace(const FVector& InRayOrigin, const FVector& InRayEnd, FVector& OutHitLocation)
 {
 	FVector Start = InRayOrigin;
 	FVector End = InRayEnd;
 
 	// Cache a copy of the world pointer
-	UWorld* World = ViewportClient->GetWorld();
+	UWorld* World = GetWorld();
 
 	static FName TraceTag = FName(TEXT("LandscapeTrace"));
 	TArray<FHitResult> Results;
@@ -2863,7 +2847,8 @@ void FEdModeLandscape::ForceRealTimeViewports(const bool bEnable, const bool bSt
 					Viewport.SetRealtime(bEnable, bStoreCurrentState);
 
 					// @todo vreditor: Force game view to true in VREditor since we can't use hitproxies and debug objects yet
-					if (IVREditorModule::Get().IsVREditorModeActive())
+					UVREditorMode* VREditorMode = Cast<UVREditorMode>( GEditor->GetEditorWorldExtensionsManager()->GetEditorWorldExtensions( GetWorld() )->FindExtension( UVREditorMode::StaticClass() ) );
+					if( VREditorMode != nullptr && VREditorMode->IsActive())
 					{
 						Viewport.SetGameView(true);
 					} 
diff --git a/Engine/Source/Editor/LevelEditor/Private/LevelEditorActions.cpp b/Engine/Source/Editor/LevelEditor/Private/LevelEditorActions.cpp
--- a/Engine/Source/Editor/LevelEditor/Private/LevelEditorActions.cpp
+++ b/Engine/Source/Editor/LevelEditor/Private/LevelEditorActions.cpp
@@ -2122,7 +2122,7 @@ void FLevelEditorActionCallbacks::ToggleVR()
 bool FLevelEditorActionCallbacks::ToggleVR_CanExecute()
 {
 	IVREditorModule& VREditorModule = IVREditorModule::Get();
-	return VREditorModule.IsVREditorAvailable() && !GEditor->bIsSimulatingInEditor;
+	return VREditorModule.IsVREditorAvailable();
 }
 
 
@@ -2132,13 +2132,6 @@ bool FLevelEditorActionCallbacks::ToggleVR_IsChecked()
 	return VREditorModule.IsVREditorEnabled();
 }
 
-
-bool FLevelEditorActionCallbacks::ToggleVR_IsVisible()
-{
-	return GetDefault<UEditorExperimentalSettings>()->bEnableVREditing;
-}
-
-
 bool FLevelEditorActionCallbacks::CanSelectGameModeBlueprint()
 {
 	bool bCheckOutNeeded = false;
diff --git a/Engine/Source/Editor/LevelEditor/Private/LevelEditorToolBar.cpp b/Engine/Source/Editor/LevelEditor/Private/LevelEditorToolBar.cpp
--- a/Engine/Source/Editor/LevelEditor/Private/LevelEditorToolBar.cpp
+++ b/Engine/Source/Editor/LevelEditor/Private/LevelEditorToolBar.cpp
@@ -1269,7 +1269,7 @@ TSharedRef< SWidget > FLevelEditorToolBar::MakeLevelEditorToolBar( const TShared
 			FSlateIcon(FEditorStyle::GetStyleSetName(), "LevelEditor.EditMatinee") 
 			);
 
-		ToolbarBuilder.AddToolBarButton( FLevelEditorCommands::Get().ToggleVR, NAME_None, LOCTEXT("ToggleVR", "VR") );
+		ToolbarBuilder.AddToolBarButton( FLevelEditorCommands::Get().ToggleVR, NAME_None, LOCTEXT("ToggleVR", "VR Mode") );
 	}
 	ToolbarBuilder.EndSection();
 	
diff --git a/Engine/Source/Editor/LevelEditor/Private/SLevelViewport.cpp b/Engine/Source/Editor/LevelEditor/Private/SLevelViewport.cpp
--- a/Engine/Source/Editor/LevelEditor/Private/SLevelViewport.cpp
+++ b/Engine/Source/Editor/LevelEditor/Private/SLevelViewport.cpp
@@ -1541,13 +1541,14 @@ EVisibility SLevelViewport::OnGetViewportContentVisibility() const
 	{
 		return BaseVisibility;
 	}
-	return IsPlayInEditorViewportActive() && IsImmersive() ? EVisibility::Collapsed : EVisibility::Visible;
+
+	return ( ( IsPlayInEditorViewportActive() && IsImmersive() ) || GEngine->IsStereoscopic3D( ActiveViewport.Get() ) ) ? EVisibility::Collapsed : EVisibility::Visible;
 }
 
 EVisibility SLevelViewport::GetToolBarVisibility() const
 {
-	// Do not show the toolbar if this viewport has a play in editor session
-	return IsPlayInEditorViewportActive() ? EVisibility::Collapsed : OnGetViewportContentVisibility();
+	// Do not show the toolbar if this viewport has a play in editor session, or we're in the VR Editor
+	return ( IsPlayInEditorViewportActive() || GEngine->IsStereoscopic3D( ActiveViewport.Get() ) ) ? EVisibility::Collapsed : OnGetViewportContentVisibility();
 }
 
 EVisibility SLevelViewport::GetMaximizeToggleVisibility() const
@@ -3494,7 +3495,7 @@ void SLevelViewport::StartPlayInEditorSession(UGameViewportClient* PlayClient, c
 	bPIEHasFocus = ActiveViewport->HasMouseCapture();
 
 
-	if(EditorPlayInSettings->ShowMouseControlLabel)
+	if(EditorPlayInSettings->ShowMouseControlLabel && !GEngine->IsStereoscopic3D( ActiveViewport.Get() ) )
 	{
 		ELabelAnchorMode AnchorMode = EditorPlayInSettings->MouseControlLabelPosition.GetValue();
 
@@ -3650,31 +3651,41 @@ void SLevelViewport::EndPlayInEditorSession()
 
 	if( IsPlayInEditorViewportActive() )
 	{
-		ActiveViewport->OnPlayWorldViewportSwapped(*InactiveViewport);
+		{
+			TSharedPtr<FSceneViewport> GameViewport = ActiveViewport;
+			ActiveViewport = InactiveViewport;
+			ActiveViewport->OnPlayWorldViewportSwapped( *GameViewport );
 
-		// Play in editor viewport was active, swap back to our level editor viewport
-		ActiveViewport->SetViewportClient( nullptr );
+			// Play in editor viewport was active, swap back to our level editor viewport
+			GameViewport->SetViewportClient( nullptr );
 
-		// We should be the only thing holding on to viewports
-		check( ActiveViewport.IsUnique() );
+			// We should be the only thing holding on to viewports
+			check( GameViewport.IsUnique() );
+		}
 
-		ActiveViewport = InactiveViewport;
-	
 		// Ensure our active viewport is for level editing
 		check( ActiveViewport->GetClient() == LevelViewportClient.Get() );
 
-		// Restore camera settings that may be adversely affected by PIE
-		LevelViewportClient->RestoreCameraFromPIE();
-		RedrawViewport(true);
+		// If we're going back to VR Editor, refresh the level viewport's render target so the HMD will present frames here
+		if( GEngine->IsStereoscopic3D( ActiveViewport.Get() ) )
+		{
+			ActiveViewport->UpdateViewportRHI( false, ActiveViewport->GetSizeXY().X, ActiveViewport->GetSizeXY().Y, ActiveViewport->GetWindowMode(), PF_Unknown );
+		}
+		else
+		{
+			// Restore camera settings that may be adversely affected by PIE
+			LevelViewportClient->RestoreCameraFromPIE();
+			RedrawViewport(true);
+
+			// Remove camera roll from any PIE camera applied in this viewport. A rolled camera is hard to use for editing
+			LevelViewportClient->RemoveCameraRoll();
+		}
 	}
 	else
 	{
 		InactiveViewport->SetViewportClient( nullptr );
 	}
 
-	// Remove camera roll from any PIE camera applied in this viewport. A rolled camera is hard to use for editing
-	LevelViewportClient->RemoveCameraRoll();
-
 	// Reset the inactive viewport
 	InactiveViewport.Reset();
 
@@ -3747,7 +3758,7 @@ void SLevelViewport::SwapViewportsForPlayInEditor()
 	ULevelEditorPlaySettings const* EditorPlayInSettings = GetDefault<ULevelEditorPlaySettings>();
 	check(EditorPlayInSettings);
 	
-	if(EditorPlayInSettings->ShowMouseControlLabel)
+	if(EditorPlayInSettings->ShowMouseControlLabel && !GEngine->IsStereoscopic3D( ActiveViewport.Get() ) )
 	{
 		ELabelAnchorMode AnchorMode = EditorPlayInSettings->MouseControlLabelPosition.GetValue();
 		
diff --git a/Engine/Source/Editor/LevelEditor/Public/LevelEditorActions.h b/Engine/Source/Editor/LevelEditor/Public/LevelEditorActions.h
--- a/Engine/Source/Editor/LevelEditor/Public/LevelEditorActions.h
+++ b/Engine/Source/Editor/LevelEditor/Public/LevelEditorActions.h
@@ -654,7 +654,6 @@ public:
 	static void ToggleVR();
 	static bool ToggleVR_CanExecute();
 	static bool ToggleVR_IsChecked();
-	static bool ToggleVR_IsVisible();
 
 	/** Opens delta transform */
 	static void DeltaTransform();
diff --git a/Engine/Source/Editor/MeshPaint/Private/MeshPaintEdMode.cpp b/Engine/Source/Editor/MeshPaint/Private/MeshPaintEdMode.cpp
--- a/Engine/Source/Editor/MeshPaint/Private/MeshPaintEdMode.cpp
+++ b/Engine/Source/Editor/MeshPaint/Private/MeshPaintEdMode.cpp
@@ -27,7 +27,6 @@
 
 #include "StaticMeshResources.h"
 #include "VREditorMode.h"
-#include "IVREditorModule.h"
 #include "ScopedTransaction.h"
 #include "MeshPaintRendering.h"
 #include "Toolkits/ToolkitManager.h"
@@ -46,11 +45,10 @@
 #include "SMeshPaint.h"
 #include "ComponentReregisterContext.h"
 #include "MeshPaintAdapterFactory.h"
-
+#include "EditorWorldExtension.h"
 #include "ViewportWorldInteraction.h"
 #include "ViewportInteractableInterface.h"
 #include "VREditorInteractor.h"
-#include "EditorWorldManager.h"
 #include "UniquePtr.h"
 
 #define LOCTEXT_NAMESPACE "MeshPaint_Mode"
@@ -257,18 +255,18 @@ void FEdModeMeshPaint::Enter()
 		ApplyOrRemoveForceBestLOD(/*bApply=*/ true, PaintingMeshLODIndex);
 	}
 
-	TSharedPtr<FEditorWorldWrapper> EditorWorld = GEditor->GetEditorWorldManager()->GetEditorWorldWrapper(GetWorld());
-	if(EditorWorld.IsValid())
+	UEditorWorldExtensionCollection* ExtensionCollection = GEditor->GetEditorWorldExtensionsManager()->GetEditorWorldExtensions(GetWorld());
+	if(ExtensionCollection != nullptr)
 	{
 		// Register to find out about VR input events
-		UViewportWorldInteraction* WorldInteraction = EditorWorld->GetViewportWorldInteraction();
-		if(WorldInteraction != nullptr)
+		UViewportWorldInteraction* ViewportWorldInteraction = Cast<UViewportWorldInteraction>(ExtensionCollection->FindExtension(UViewportWorldInteraction::StaticClass()));
+		if(ViewportWorldInteraction != nullptr)
 		{
-			WorldInteraction->OnViewportInteractionInputAction().RemoveAll(this);
-			WorldInteraction->OnViewportInteractionInputAction().AddRaw(this, &FEdModeMeshPaint::OnVRAction);
+			ViewportWorldInteraction->OnViewportInteractionInputAction().RemoveAll(this);
+			ViewportWorldInteraction->OnViewportInteractionInputAction().AddRaw(this, &FEdModeMeshPaint::OnVRAction);
 
 			// Hide the VR transform gizmo while we're in mesh paint mode.  It sort of gets in the way of painting.
-			WorldInteraction->SetTransformGizmoVisible(false);
+			ViewportWorldInteraction->SetTransformGizmoVisible(false);
 		}
 	}
 }
@@ -276,20 +274,17 @@ void FEdModeMeshPaint::Enter()
 /** FEdMode: Called when the mode is exited */
 void FEdModeMeshPaint::Exit()
 {
-	if(IVREditorModule::IsAvailable() && IVREditorModule::Get().IsVREditorModeActive())
+	UEditorWorldExtensionCollection* ExtensionCollection = GEditor->GetEditorWorldExtensionsManager()->GetEditorWorldExtensions(GetWorld());
+	if (ExtensionCollection != nullptr)
 	{
-		TSharedPtr<FEditorWorldWrapper> EditorWorldWrapper = GEditor->GetEditorWorldManager()->GetEditorWorldWrapper(GetWorld(), false);
-		if (EditorWorldWrapper.IsValid())
+		UViewportWorldInteraction* ViewportWorldInteraction = Cast<UViewportWorldInteraction>(ExtensionCollection->FindExtension(UViewportWorldInteraction::StaticClass()));
+		if(ViewportWorldInteraction != nullptr )
 		{
-			UViewportWorldInteraction* WorldInteraction = EditorWorldWrapper->GetViewportWorldInteraction();
-			if (WorldInteraction != nullptr)
-			{
-				// Restore the transform gizmo visibility
-				WorldInteraction->SetTransformGizmoVisible(true);
+			// Restore the transform gizmo visibility
+			ViewportWorldInteraction->SetTransformGizmoVisible( true );
 
-				// Unregister from event handlers
-				WorldInteraction->OnViewportInteractionInputAction().RemoveAll(this);
-			}
+			// Unregister from event handlers
+			ViewportWorldInteraction->OnViewportInteractionInputAction().RemoveAll( this );
 		}
 	}
 
@@ -2753,11 +2748,9 @@ void FEdModeMeshPaint::Render( const FSceneView* View, FViewport* Viewport, FPri
 		static TArray<FPaintRay> PaintRays;
 		PaintRays.Reset();
 		
-		UVREditorMode* VREditorMode = IVREditorModule::Get().GetVREditorMode();
-
+		UVREditorMode* VREditorMode = Cast<UVREditorMode>( GEditor->GetEditorWorldExtensionsManager()->GetEditorWorldExtensions( GetWorld() )->FindExtension( UVREditorMode::StaticClass() ) );
 		// Check to see if VREditorMode is active. If so, we're painting with interactor
 		bool bIsInVRMode = false;
-		if ( IVREditorModule::IsAvailable() )
 		{
 			if ( VREditorMode != nullptr && VREditorMode->IsFullyInitialized() && VREditorMode->IsActive())
 			{
@@ -3680,7 +3673,7 @@ EMeshPaintAction::Type FEdModeMeshPaint::GetPaintAction(FViewport* InViewport)
 		// When painting using VR, allow the modifier button to activate Erase mode
 		if( PaintingWithInteractorInVR )
 		{
-			UVREditorMode* VREditorMode = GEditor->GetEditorWorldManager()->GetEditorWorldWrapper( GetWorld() )->GetVREditorMode();
+			UVREditorMode* VREditorMode = Cast<UVREditorMode>( GEditor->GetEditorWorldExtensionsManager()->GetEditorWorldExtensions( GetWorld() )->FindExtension( UVREditorMode::StaticClass() ) );
 			UVREditorInteractor* VRInteractor = Cast<UVREditorInteractor>( PaintingWithInteractorInVR );
 			if( VREditorMode != nullptr && VRInteractor != nullptr )
 			{
@@ -5432,11 +5425,11 @@ TWeakObjectPtr<AActor> FEdModeMeshPaint::GetEditingActor() const
 void FEdModeMeshPaint::OnVRAction( class FEditorViewportClient& ViewportClient, UViewportInteractor* Interactor, 
 	const FViewportActionKeyInput& Action, bool& bOutIsInputCaptured, bool& bWasHandled )
 {
-	UVREditorInteractor* VRInteractor = Cast<UVREditorInteractor>(Interactor);
-	if (IVREditorModule::Get().IsVREditorModeActive() && VRInteractor != nullptr)
+	UVREditorMode* VREditorMode = Cast<UVREditorMode>( GEditor->GetEditorWorldExtensionsManager()->GetEditorWorldExtensions( GetWorld() )->FindExtension( UVREditorMode::StaticClass() ) );
+	UVREditorInteractor* VRInteractor = Cast<UVREditorInteractor>( Interactor );
+	if( VREditorMode != nullptr && VREditorMode->IsActive() && VRInteractor != nullptr && VRInteractor->GetDraggingMode() == EViewportInteractionDraggingMode::Nothing)
 	{
-		UVREditorMode* VREditorMode = IVREditorModule::Get().GetVREditorMode();
-		if (Action.ActionType == ViewportWorldActionTypes::SelectAndMove_LightlyPressed)
+		if( Action.ActionType == ViewportWorldActionTypes::SelectAndMove )
 		{
 			if (!bIsPainting && Action.Event == IE_Pressed && !VRInteractor->IsHoveringOverPriorityType())
 			{
@@ -5474,10 +5467,6 @@ void FEdModeMeshPaint::OnVRAction( class FEditorViewportClient& ViewportClient,
 					bWasHandled = true;
 					bOutIsInputCaptured = true;
 
-					// Don't allow a "full press" to happen at all with this trigger pull.  We don't need full press for anything, and it
-					// would interrupt our light press.
-					VRInteractor->SetAllowTriggerFullPress(false);
-
 					StartPainting();
 					PaintingWithInteractorInVR = Interactor;
 
@@ -5507,6 +5496,12 @@ void FEdModeMeshPaint::OnVRAction( class FEditorViewportClient& ViewportClient,
 				bWasHandled = true;
 				bOutIsInputCaptured = false;
 			}
+
+			else if( bIsPainting )
+			{
+				// A different hand might be painting, so absorb the input
+				bOutIsInputCaptured = bWasHandled = ( Action.Event == IE_Released ? false : true );
+			}
 		}
 	}
 }
diff --git a/Engine/Source/Editor/MovieSceneTools/Private/Sections/ThumbnailSection.cpp b/Engine/Source/Editor/MovieSceneTools/Private/Sections/ThumbnailSection.cpp
--- a/Engine/Source/Editor/MovieSceneTools/Private/Sections/ThumbnailSection.cpp
+++ b/Engine/Source/Editor/MovieSceneTools/Private/Sections/ThumbnailSection.cpp
@@ -15,6 +15,7 @@
 #include "Widgets/Text/SInlineEditableTextBlock.h"
 #include "PropertyEditorModule.h"
 #include "IDetailsView.h"
+#include "IVREditorModule.h"
 
 
 #define LOCTEXT_NAMESPACE "FThumbnailSection"
@@ -201,7 +202,7 @@ int32 FThumbnailSection::OnPaintSection( FSequencerSectionPainter& InPainter ) c
 
 	static const float SectionThumbnailPadding = 4.f;
 
-	const ESlateDrawEffect DrawEffects = InPainter.bParentEnabled ? ESlateDrawEffect::None : ESlateDrawEffect::DisabledEffect;
+	ESlateDrawEffect DrawEffects = InPainter.bParentEnabled ? ESlateDrawEffect::None : ESlateDrawEffect::DisabledEffect;
 
 	int32 LayerId = InPainter.LayerId;
 
@@ -240,13 +241,22 @@ int32 FThumbnailSection::OnPaintSection( FSequencerSectionPainter& InPainter ) c
 
 		if (Fade <= 1.f)
 		{
+			DrawEffects |= ESlateDrawEffect::NoGamma;
+
+			if (IVREditorModule::Get().IsVREditorModeActive())
+			{
+				// In VR editor every widget is in the world and gamma corrected by the scene renderer.  Thumbnails will have already been gamma
+				// corrected and so they need to be reversed
+				DrawEffects |= ESlateDrawEffect::ReverseGamma;
+			}
+
 			FSlateDrawElement::MakeViewport(
 				InPainter.DrawElements,
 				LayerId,
 				PaintGeometry,
 				Thumbnail,
 				ThumbnailClipRect,
-				DrawEffects | ESlateDrawEffect::NoGamma,
+				DrawEffects,
 				FLinearColor(1.f, 1.f, 1.f, 1.f - Fade)
 				);
 		}
diff --git a/Engine/Source/Editor/MovieSceneTools/Private/TrackEditors/TransformTrackEditor.cpp b/Engine/Source/Editor/MovieSceneTools/Private/TrackEditors/TransformTrackEditor.cpp
--- a/Engine/Source/Editor/MovieSceneTools/Private/TrackEditors/TransformTrackEditor.cpp
+++ b/Engine/Source/Editor/MovieSceneTools/Private/TrackEditors/TransformTrackEditor.cpp
@@ -437,20 +437,32 @@ TSharedRef<ISequencerSection> F3DTransformTrackEditor::MakeSectionInterface(UMov
 void F3DTransformTrackEditor::OnPreTransformChanged( UObject& InObject )
 {
 	float AutoKeyTime = GetTimeForKey();
-
-	if( IsAllowedToAutoKey() )
+	AActor* Actor = Cast<AActor>(&InObject);
+	// If Sequencer is allowed to autokey and we are clicking on an Actor that can be autokeyed
+	if( IsAllowedToAutoKey() && (Actor && !Actor->IsEditorOnly() ) )
 	{
-		AActor* Actor = nullptr;
+		AActor* ActorThatChanged = nullptr;
 		USceneComponent* SceneComponentThatChanged = nullptr;
-		GetActorAndSceneComponentFromObject(&InObject, Actor, SceneComponentThatChanged);
+		GetActorAndSceneComponentFromObject(&InObject, ActorThatChanged, SceneComponentThatChanged);
+		USceneComponent* SceneComponent = nullptr;
+		
+		if( ActorThatChanged && ActorThatChanged->GetRootComponent() )
+		{
+			SceneComponent = ActorThatChanged->GetRootComponent();
+		}
+		else
+		{
+			// If the object wasn't an actor attempt to get it directly as a scene component 
+			SceneComponent = Cast<USceneComponent>( &InObject );
+		}
 
 		if( SceneComponentThatChanged )
 		{
 			// Cache off the existing transform so we can detect which components have changed
 			// and keys only when something has changed
 			FTransformData Transform( SceneComponentThatChanged );
-
-			ObjectToExistingTransform.Add( &InObject, Transform );
+			
+			ObjectToExistingTransform.Add(&InObject, Transform);
 		}
 	}
 }
@@ -477,7 +489,8 @@ void F3DTransformTrackEditor::OnTransformChanged( UObject& InObject )
 	USceneComponent* SceneComponentThatChanged = nullptr;
 	GetActorAndSceneComponentFromObject(&InObject, Actor, SceneComponentThatChanged);
 
-	if( SceneComponentThatChanged != nullptr && Actor != nullptr )
+	// If the Actor that just finished transforming doesn't have autokey disabled
+	if( SceneComponentThatChanged != nullptr && (Actor && !Actor->IsEditorOnly()))
 	{
 		// Find an existing transform if possible.  If one exists we will compare against the new one to decide what components of the transform need keys
 		FTransformData ExistingTransform = ObjectToExistingTransform.FindRef( &InObject );
diff --git a/Engine/Source/Editor/Sequencer/Private/Sequencer.h b/Engine/Source/Editor/Sequencer/Private/Sequencer.h
--- a/Engine/Source/Editor/Sequencer/Private/Sequencer.h
+++ b/Engine/Source/Editor/Sequencer/Private/Sequencer.h
@@ -372,7 +372,6 @@ public:
 	void RecordSelectedActors();
 
 	/** Functions to push on to the transport controls we use */
-	FReply OnPlay(bool bTogglePlay=true, float InPlayRate=1.f);
 	FReply OnRecord();
 	FReply OnStepForward();
 	FReply OnStepBackward();
@@ -393,6 +392,7 @@ public:
 	/** Delegate handler for recording finishing */
 	void HandleRecordingFinished(UMovieSceneSequence* Sequence);
 
+	/** Whether or not playback is looping */
 	bool IsLooping() const;
 
 	/** Set the new global time, accounting for looping options */
@@ -614,11 +614,14 @@ public:
 	virtual TSharedRef<INumericTypeInterface<float>> GetNumericTypeInterface() override;
 	virtual TSharedRef<INumericTypeInterface<float>> GetZeroPadNumericTypeInterface() override;
 	virtual TSharedRef<SWidget> MakeTransportControls(bool bExtended) override;
+	virtual FReply OnPlay(bool bTogglePlay = true, float InPlayRate = 1.f) override;
+	virtual void Pause() override;
 	virtual TSharedRef<SWidget> MakeTimeRange(const TSharedRef<SWidget>& InnerContent, bool bShowWorkingRange, bool bShowViewRange, bool bShowPlaybackRange) override;
 	virtual UObject* FindSpawnedObjectOrTemplate(const FGuid& BindingId) override;
 	virtual FGuid MakeNewSpawnable(UObject& SourceObject) override;
 	virtual bool IsReadOnly() const override;
 	virtual void ExternalSelectionHasChanged() override { SynchronizeSequencerSelectionWithExternalSelection(); }
+	virtual USequencerSettings* GetSequencerSettings() override { return Settings; }
 
 public:
 
@@ -745,7 +748,6 @@ protected:
 	void Rewind();
 	void ShuttleForward();
 	void ShuttleBackward();
-	void Pause();
 	void StepForward();
 	void StepBackward();
 	void StepToNextKey();
diff --git a/Engine/Source/Editor/Sequencer/Private/SequencerKeyActor.cpp b/Engine/Source/Editor/Sequencer/Private/SequencerKeyActor.cpp
new file mode 100644
--- /dev/null
+++ b/Engine/Source/Editor/Sequencer/Private/SequencerKeyActor.cpp
@@ -0,0 +1,95 @@
+// Copyright 1998-2016 Epic Games, Inc. All Rights Reserved.
+
+#include "SequencerKeyActor.h"
+#include "Components/StaticMeshComponent.h"
+#include "Sections/MovieScene3DTransformSection.h"
+#include "SequencerEdMode.h"
+#include "Materials/MaterialInstance.h" 
+#include "Materials/MaterialInstanceDynamic.h"
+#include "Engine/Selection.h"
+#include "EditorModeTools.h"
+#include "EditorModeManager.h"
+#include "ModuleManager.h"
+#include "EditorViewportClient.h"
+#include "UnrealClient.h"
+
+ASequencerKeyActor::ASequencerKeyActor()
+	: Super()
+{
+	UStaticMesh* KeyEditorMesh = LoadObject<UStaticMesh>(nullptr, TEXT("/Engine/VREditor/TransformGizmo/SM_Sequencer_Key"));
+	check(KeyEditorMesh != nullptr);
+	UMaterial* KeyEditorMaterial = LoadObject<UMaterial>(nullptr, TEXT("/Engine/VREditor/TransformGizmo/Main"));
+	check(KeyEditorMaterial != nullptr);
+
+	const bool bTransient = true;
+	USceneComponent* SceneComponent = CreateDefaultSubobject<USceneComponent>(TEXT("SceneComponent"), bTransient);
+	check(SceneComponent != nullptr);
+	this->RootComponent = SceneComponent;
+
+	KeyMeshComponent = CreateDefaultSubobject<UStaticMeshComponent>(TEXT("KeyMesh"));
+	KeyMeshComponent->SetMobility(EComponentMobility::Movable);
+	KeyMeshComponent->SetupAttachment(RootComponent);
+	KeyMeshComponent->SetStaticMesh(KeyEditorMesh);
+	KeyMeshComponent->CreateAndSetMaterialInstanceDynamicFromMaterial(0, KeyEditorMaterial);
+
+	KeyMeshComponent->SetCollisionEnabled(ECollisionEnabled::QueryOnly);
+	KeyMeshComponent->SetCollisionResponseToAllChannels(ECR_Ignore);
+	KeyMeshComponent->SetCollisionResponseToChannel(ECC_Visibility, ECR_Block);
+	
+	KeyMeshComponent->bGenerateOverlapEvents = false;
+	KeyMeshComponent->SetCanEverAffectNavigation(false);
+	KeyMeshComponent->bCastDynamicShadow = false;
+	KeyMeshComponent->bCastStaticShadow = false;
+	KeyMeshComponent->bAffectDistanceFieldLighting = false;
+	KeyMeshComponent->bAffectDynamicIndirectLighting = false;
+}
+
+
+void ASequencerKeyActor::PostEditMove(bool bFinished)
+{
+	// Push the key's transform to the Sequencer track
+	PropagateKeyChange();
+	Super::PostEditMove(bFinished);
+}
+
+void ASequencerKeyActor::SetKeyData(class UMovieScene3DTransformSection* NewTrackSection, float NewKeyTime)
+{
+	TrackSection = NewTrackSection;
+	KeyTime = NewKeyTime;
+	// Associate the currently selected actor with this key
+	AssociatedActor = GEditor->GetSelectedActors()->GetTop<AActor>();
+	// Draw a single transform track based on the data from this key
+	FEditorViewportClient* ViewportClient = StaticCast<FEditorViewportClient*>(GEditor->GetActiveViewport()->GetClient());
+	if (ViewportClient != nullptr)
+	{
+		FSequencerEdMode* SequencerEdMode = (FSequencerEdMode*)(ViewportClient->GetModeTools()->GetActiveMode(FSequencerEdMode::EM_SequencerMode));
+		SequencerEdMode->DrawMeshTransformTrailFromKey(this);
+	}
+}
+
+
+void ASequencerKeyActor::PropagateKeyChange()
+{
+	if (TrackSection != nullptr)
+	{
+		// Mark the track section as dirty
+		TrackSection->Modify();
+
+		// Update the translation keys
+		FRichCurve& TransXCurve = TrackSection->GetTranslationCurve(EAxis::X);
+		FRichCurve& TransYCurve = TrackSection->GetTranslationCurve(EAxis::Y);
+		FRichCurve& TransZCurve = TrackSection->GetTranslationCurve(EAxis::Z);
+		TransXCurve.UpdateOrAddKey(KeyTime, GetActorTransform().GetLocation().X);
+		TransYCurve.UpdateOrAddKey(KeyTime, GetActorTransform().GetLocation().Y);
+		TransZCurve.UpdateOrAddKey(KeyTime, GetActorTransform().GetLocation().Z);
+
+		// Draw a single transform track based on the data from this key
+		FEditorViewportClient* ViewportClient = StaticCast<FEditorViewportClient*>(GEditor->GetActiveViewport()->GetClient());
+		if (ViewportClient != nullptr)
+		{
+			FSequencerEdMode* SequencerEdMode = (FSequencerEdMode*)(ViewportClient->GetModeTools()->GetActiveMode(FSequencerEdMode::EM_SequencerMode));
+			SequencerEdMode->DrawMeshTransformTrailFromKey(this);
+		}
+	}
+}
+
diff --git a/Engine/Source/Editor/Sequencer/Private/SequencerKeyActor.h b/Engine/Source/Editor/Sequencer/Private/SequencerKeyActor.h
new file mode 100644
--- /dev/null
+++ b/Engine/Source/Editor/Sequencer/Private/SequencerKeyActor.h
@@ -0,0 +1,62 @@
+// Copyright 1998-2016 Epic Games, Inc. All Rights Reserved.
+
+#pragma once
+
+#include "Engine/StaticMeshActor.h"
+#include "SequencerKeyActor.generated.h"
+
+UCLASS()
+class SEQUENCER_API ASequencerKeyActor : public AActor
+{
+	GENERATED_BODY()
+
+public:
+
+	ASequencerKeyActor();
+
+	/** AActor overrides */
+	virtual void PostEditMove(bool bFinished) override;
+	virtual bool IsEditorOnly() const final
+	{
+		return true;
+	}
+	/** AActor */
+
+	/** Set the track section and time associated with this key */
+	void SetKeyData(class UMovieScene3DTransformSection* NewTrackSection, float NewKeyTime);
+
+	/** Return the actor associated with this key */
+	AActor* GetAssociatedActor() const
+	{ 
+		return AssociatedActor; 
+	}
+
+	/** Returns the mesh component for this key, or nullptr if not spawned right now */
+	class UStaticMeshComponent* GetMeshComponent()
+	{
+		return KeyMeshComponent;
+	}
+
+protected:
+
+	/** Push the data from a key actor change back to Sequencer */
+	void PropagateKeyChange();
+
+private:
+
+	/** The key mesh */
+	UPROPERTY()
+	class UStaticMeshComponent* KeyMeshComponent;
+
+	/** The actor whose transform was used to build this key */
+	UPROPERTY()
+	AActor*	AssociatedActor;
+
+	/** The track section this key resides on */
+	UPROPERTY()
+	class UMovieScene3DTransformSection* TrackSection;
+
+	/** The time this key is associated with in Sequencer */
+	UPROPERTY()
+	float KeyTime;
+};
\ No newline at end of file
diff --git a/Engine/Source/Editor/Sequencer/Private/SequencerMeshTrail.cpp b/Engine/Source/Editor/Sequencer/Private/SequencerMeshTrail.cpp
new file mode 100644
--- /dev/null
+++ b/Engine/Source/Editor/Sequencer/Private/SequencerMeshTrail.cpp
@@ -0,0 +1,159 @@
+// Copyright 1998-2017 Epic Games, Inc. All Rights Reserved.
+
+#include "SequencerMeshTrail.h"
+#include "SequencerKeyActor.h"
+#include "EditorWorldExtension.h"
+#include "ViewportWorldInteraction.h"
+#include "Materials/MaterialInstance.h" 
+#include "Materials/MaterialInstanceDynamic.h"
+#include "Components/BillboardComponent.h"
+#include "Components/StaticMeshComponent.h"
+#include "MovieScene3DTransformSection.h"
+#include "TimerManager.h"
+#include "Editor.h"
+
+ASequencerMeshTrail::ASequencerMeshTrail()
+	: Super()
+{
+	TrailTime = 0.0f;
+	MaxTrailTime = 1.0f;
+	UBillboardComponent* PrimitiveRootComponent = CreateDefaultSubobject<UBillboardComponent>(TEXT("RootComponent"));
+	PrimitiveRootComponent->bSelectable = false;
+	PrimitiveRootComponent->SetVisibility(false, false);
+	SetRootComponent(PrimitiveRootComponent);
+
+	if (!IsTemplate() && GEditor)
+	{
+		GEditor->GetTimerManager()->SetTimer(TrailUpdate, FTimerDelegate::CreateUObject(this, &ASequencerMeshTrail::UpdateTrailAppearance, 0.0005f), 0.0005f, true);
+	}
+}
+
+void ASequencerMeshTrail::Cleanup()
+{
+	// Destroy all the key actors this trail created
+	for (FKeyActorData KeyMesh : KeyMeshActors)
+	{
+		KeyMesh.KeyActor->Destroy();
+	}
+	this->Destroy();
+}
+
+void ASequencerMeshTrail::AddKeyMeshActor(float KeyTime, const FTransform KeyTransform, class UMovieScene3DTransformSection* TrackSection)
+{
+	ASequencerKeyActor* KeyMeshActor = nullptr;
+	UStaticMesh* KeyEditorMesh = nullptr;
+	UMaterialInstance* KeyEditorMaterial = nullptr;
+	MaxTrailTime = FMath::Max(KeyTime, MaxTrailTime);
+	FKeyActorData* KeyPtr = KeyMeshActors.FindByPredicate([KeyTime](const FKeyActorData InKey)
+	{
+		return FMath::IsNearlyEqual(KeyTime, InKey.Time);
+	});
+
+	// If we don't currently have an actor for this time, create one
+	if (KeyPtr == nullptr)
+	{
+		UViewportWorldInteraction* WorldInteraction = Cast<UViewportWorldInteraction>( GEditor->GetEditorWorldExtensionsManager()->GetEditorWorldExtensions( GetWorld() )->FindExtension( UViewportWorldInteraction::StaticClass() ) );
+		if( WorldInteraction != nullptr )
+		{
+			KeyMeshActor = WorldInteraction->SpawnTransientSceneActor<ASequencerKeyActor>( TEXT( "KeyMesh" ), false );
+			KeyMeshActor->SetActorTransform(KeyTransform);
+			// Destroy all the frame components this trail created
+			for (FFrameComponentData FrameMesh : FrameMeshComponents)
+			{
+				FrameMesh.FrameComponent->DestroyComponent();
+			}
+			FrameMeshComponents.Reset();
+			KeyMeshActor->SetKeyData(TrackSection, KeyTime);
+			KeyMeshActor->SetOwner(this);
+			FKeyActorData NewKey = FKeyActorData(KeyTime, KeyMeshActor);
+			KeyMeshActors.Add(NewKey);
+		}
+	}
+	else
+	{
+		// Just update the transform
+		KeyPtr->KeyActor->SetActorTransform(KeyTransform);
+	}
+}
+
+void ASequencerMeshTrail::AddFrameMeshComponent(const float FrameTime, const FTransform FrameTransform)
+{
+	UStaticMeshComponent* FrameMeshComponent = nullptr;
+	UStaticMesh* FrameEditorMesh = nullptr;
+	UMaterial* FrameEditorMaterial = nullptr;
+
+	FFrameComponentData* FramePtr = FrameMeshComponents.FindByPredicate([FrameTime](const FFrameComponentData InFrame)
+	{
+		return FMath::IsNearlyEqual(FrameTime, InFrame.Time);
+	});
+
+	// If we don't currently have a component for this time, create one
+	if (FramePtr == nullptr)
+	{
+		FrameMeshComponent = NewObject<UStaticMeshComponent>(this);
+		this->AddOwnedComponent(FrameMeshComponent);
+		FrameMeshComponent->RegisterComponent();
+		FrameEditorMesh = LoadObject<UStaticMesh>(nullptr, TEXT("/Engine/VREditor/TransformGizmo/SM_Sequencer_Node"));
+		check(FrameEditorMesh != nullptr);
+		FrameEditorMaterial = LoadObject<UMaterial>(nullptr, TEXT("/Engine/VREditor/TransformGizmo/Main"));
+		check(FrameEditorMaterial != nullptr);
+		FrameMeshComponent->SetStaticMesh(FrameEditorMesh);
+		FrameMeshComponent->CreateAndSetMaterialInstanceDynamicFromMaterial(0, FrameEditorMaterial);
+		FrameMeshComponent->SetMobility(EComponentMobility::Movable);
+		FrameMeshComponent->SetWorldTransform(FrameTransform);
+		FrameMeshComponent->SetCastShadow(false);
+		FrameMeshComponent->SetCollisionEnabled(ECollisionEnabled::NoCollision);
+		FrameMeshComponent->bSelectable = false;
+		FFrameComponentData NewFrame = FFrameComponentData(FrameTime, FrameMeshComponent);
+		FrameMeshComponents.Add(NewFrame);
+	}
+	else
+	{
+		// Just update the transform
+		FramePtr->FrameComponent->SetWorldTransform(FrameTransform);
+	}
+}
+
+void ASequencerMeshTrail::UpdateTrailAppearance(float UpdateTime)
+{
+	if (FMath::IsNearlyEqual(MaxTrailTime, 0.0f))
+	{
+		MaxTrailTime = 1.0f;
+	}
+	// TODO: No hardcoded glow values
+	TrailTime = FMath::Fmod(TrailTime + UpdateTime, MaxTrailTime);
+	for (FFrameComponentData Frame : FrameMeshComponents)
+	{
+		UMaterialInstanceDynamic* FrameMaterial = Cast<UMaterialInstanceDynamic>(Frame.FrameComponent->GetMaterial(0));
+		if (FrameMaterial != nullptr)
+		{
+			float GlowValue = 0.0f;
+			if (FMath::IsNearlyEqual(TrailTime, Frame.Time, 0.1f))
+			{
+				GlowValue = 12.0f;
+			}
+			else
+			{
+				GlowValue = 3.0f;
+			}
+			FrameMaterial->SetScalarParameterValue("GlowAmount", GlowValue);
+		}
+	}
+	for (FKeyActorData Key : KeyMeshActors)
+	{
+		UMaterialInstanceDynamic* KeyMaterial = Cast<UMaterialInstanceDynamic>(Key.KeyActor->GetMeshComponent()->GetMaterial(0));
+		if (KeyMaterial != nullptr)
+		{
+			float GlowValue = 0.0f;
+			if (FMath::IsNearlyEqual(TrailTime, Key.Time, 0.1f))
+			{
+				GlowValue = 20.0f;
+			}
+			else
+			{
+				GlowValue = 5.0f;
+			}
+			KeyMaterial->SetScalarParameterValue("GlowAmount", GlowValue);
+		}
+	}
+}
diff --git a/Engine/Source/Editor/Sequencer/Public/ISequencer.h b/Engine/Source/Editor/Sequencer/Public/ISequencer.h
--- a/Engine/Source/Editor/Sequencer/Public/ISequencer.h
+++ b/Engine/Source/Editor/Sequencer/Public/ISequencer.h
@@ -369,6 +369,15 @@ public:
 	/** @return Returns a widget containing the sequencer's playback controls */
 	virtual TSharedRef<SWidget> MakeTransportControls(bool bExtended) = 0;
 
+	/** Play or toggle playback at the specified play rate */
+	virtual FReply OnPlay(bool bTogglePlay = true, float InPlayRate = 1.f) = 0;
+
+	/** Pause playback */
+	virtual void Pause() = 0;
+
+	/** Getter for sequencer settings */
+	virtual USequencerSettings* GetSequencerSettings() = 0;
+
 	/** Attempt to find a spawned object in the currently focused movie scene, or the template object for the specified binding ID, if possible */
 	virtual UObject* FindSpawnedObjectOrTemplate(const FGuid& BindingId) = 0;
 
diff --git a/Engine/Source/Editor/UnrealEd/Classes/Editor/EditorEngine.h b/Engine/Source/Editor/UnrealEd/Classes/Editor/EditorEngine.h
--- a/Engine/Source/Editor/UnrealEd/Classes/Editor/EditorEngine.h
+++ b/Engine/Source/Editor/UnrealEd/Classes/Editor/EditorEngine.h
@@ -55,6 +55,7 @@ class USoundBase;
 class USoundNode;
 class UTextureRenderTarget2D;
 struct FAnalyticsEventAttribute;
+class UEditorWorldExtensionManager;
 
 //
 // Things to set in mapSetBrush.
@@ -614,6 +615,12 @@ public:
 	/** The feature level we should use when loading or creating a new world */
 	ERHIFeatureLevel::Type DefaultWorldFeatureLevel;
 
+private:
+
+	/** Manager that holds all extensions paired with a world */
+	UPROPERTY()
+	UEditorWorldExtensionManager* EditorWorldExtensionsManager;
+
 protected:
 
 	/** Count of how many PIE instances are waiting to log in */
@@ -1153,6 +1160,12 @@ public:
 	/** Plays an editor sound (if the user has sounds enabled.) */
 	void PlayEditorSound( USoundBase* InSound );
 
+	/**
+	 * Returns true if currently able to play a sound and if the user has sounds enabled.
+	 */
+	bool CanPlayEditorSound() const;
+
+
 	/**
 	 * Returns the preview audio component
 	 */
@@ -1819,7 +1832,7 @@ public:
 	/**
 	*  Returns the Editors world manager instance.
 	*/
-	TSharedRef<class FEditorWorldManager> GetEditorWorldManager() { return EditorWorldManager.ToSharedRef(); }
+	UEditorWorldExtensionManager* GetEditorWorldExtensionsManager() { return EditorWorldExtensionsManager; }
 
 	// Editor specific
 
@@ -2823,9 +2836,6 @@ private:
 	/** The Timer manager for all timer delegates */
 	TSharedPtr<class FTimerManager> TimerManager;
 
-	/** Manager that holds all the worlds contexts with viewport interaction */
-	TSharedPtr<class FEditorWorldManager> EditorWorldManager;
-
 	/** The output log -> message log redirector for use during PIE */
 	TSharedPtr<class FOutputLogErrorsToMessageLogProxy> OutputLogErrorsToMessageLogProxyPtr;
 
diff --git a/Engine/Source/Editor/UnrealEd/Classes/Factories/FbxImportUI.h b/Engine/Source/Editor/UnrealEd/Classes/Factories/FbxImportUI.h
--- a/Engine/Source/Editor/UnrealEd/Classes/Factories/FbxImportUI.h
+++ b/Engine/Source/Editor/UnrealEd/Classes/Factories/FbxImportUI.h
@@ -22,8 +22,6 @@ enum EFBXImportType
 	FBXIT_SkeletalMesh UMETA(DisplayName="Skeletal Mesh"),
 	/** Select Animation if you'd like to import only animation. */
 	FBXIT_Animation UMETA(DisplayName="Animation"),
-	/** Subdivision Surface (Experimental, Early work in progress) */
-	FBXIT_SubDSurface UMETA(DisplayName="SubDSurface"),
 
 	FBXIT_MAX,
 };
@@ -55,9 +53,6 @@ public:
 	bool bImportAsSkeletal;
 	
 	/** Whether to import the incoming FBX as a Subdivision Surface (could be made a combo box together with bImportAsSkeletal) (Experimental, Early work in progress) */
-	UPROPERTY(EditAnywhere, AdvancedDisplay, Category = Mesh, meta = (ImportType = "StaticMesh|SkeletalMesh", DisplayName="Subdivison Surface"))
-	bool bImportAsSubDSurface;
-
 	/** Whether to import the mesh. Allows animation only import when importing a skeletal mesh. */
 	UPROPERTY(EditAnywhere, Category=Mesh, meta=(ImportType="SkeletalMesh"))
 	bool bImportMesh;
@@ -160,11 +155,6 @@ public:
 	void SetMeshTypeToImport()
 	{
 		MeshTypeToImport = bImportAsSkeletal ? FBXIT_SkeletalMesh : FBXIT_StaticMesh;
-
-		if(bImportAsSubDSurface)
-		{
-			MeshTypeToImport = FBXIT_SubDSurface;
-		}
 	}
 };
 
diff --git a/Engine/Source/Editor/UnrealEd/Classes/Settings/EditorExperimentalSettings.h b/Engine/Source/Editor/UnrealEd/Classes/Settings/EditorExperimentalSettings.h
--- a/Engine/Source/Editor/UnrealEd/Classes/Settings/EditorExperimentalSettings.h
+++ b/Engine/Source/Editor/UnrealEd/Classes/Settings/EditorExperimentalSettings.h
@@ -125,18 +125,6 @@ public:
 	UPROPERTY(EditAnywhere, config, Category = Tools, meta = (DisplayName = "Use OpenCL for convex hull decomposition"))
 	bool bUseOpenCLForConvexHullDecomp;
 
-	/** Enables a preview of the Unreal Editor in VR.  This adds a new tool bar button that allows you to toggle into "VR Mode" instantly.  This feature is still in development, but your feedback is appreciated! */
-	UPROPERTY(EditAnywhere, config, Category=VR, meta = (DisplayName="Enable VR Editing"))
-	bool bEnableVREditing;
-
-	/**If true, wearing a Vive or Oculus Rift headset will automatically enter VR Editing mode if Enable VR Editing is true. */
-	UPROPERTY(EditAnywhere, config, Category = VR, meta = (DisplayName = "Enable VR Mode Auto-Entry"))
-	bool bEnableAutoVREditMode;
-
-	/**If true, the tutorial window should always appear when the VR Editor is launched. */
-	UPROPERTY(EditAnywhere, config, Category = VR, meta = (DisplayName = "Always Show VR Tutorial at Start"))
-	bool bAlwaysShowVRTutorial;
-
 	/** Allows editing of potentially unsafe properties during PIE. Advanced use only - use with caution. */
 	UPROPERTY(EditAnywhere, config, Category = Tools, meta = (DisplayName = "Allow editing of potentially unsafe properties."))
 	bool bAllowPotentiallyUnsafePropertyEditing;
diff --git a/Engine/Source/Editor/UnrealEd/Private/AssetThumbnail.cpp b/Engine/Source/Editor/UnrealEd/Private/AssetThumbnail.cpp
--- a/Engine/Source/Editor/UnrealEd/Private/AssetThumbnail.cpp
+++ b/Engine/Source/Editor/UnrealEd/Private/AssetThumbnail.cpp
@@ -34,6 +34,7 @@
 #include "AssetToolsModule.h"
 #include "Styling/SlateIconFinder.h"
 #include "ClassIconFinder.h"
+#include "IVREditorModule.h"
 
 class SAssetThumbnail : public SCompoundWidget
 {
@@ -150,6 +151,9 @@ public:
 			TSharedPtr<SViewport> Viewport = 
 				SNew( SViewport )
 				.EnableGammaCorrection(false)
+				// In VR editor every widget is in the world and gamma corrected by the scene renderer.  Thumbnails will have already been gamma
+				// corrected and so they need to be reversed
+				.ReverseGammaCorrection(IVREditorModule::Get().IsVREditorModeActive())
 				.EnableBlending(true);
 
 			Viewport->SetViewportInterface( AssetThumbnail.ToSharedRef() );
diff --git a/Engine/Source/Editor/UnrealEd/Private/Editor.cpp b/Engine/Source/Editor/UnrealEd/Private/Editor.cpp
--- a/Engine/Source/Editor/UnrealEd/Private/Editor.cpp
+++ b/Engine/Source/Editor/UnrealEd/Private/Editor.cpp
@@ -89,10 +89,13 @@ FEditorDelegates::FOnEditorModeTransitioned				FEditorDelegates::EditorModeEnter
 FEditorDelegates::FOnEditorModeTransitioned				FEditorDelegates::EditorModeExit;
 FEditorDelegates::FOnPIEEvent							FEditorDelegates::PreBeginPIE;
 FEditorDelegates::FOnPIEEvent							FEditorDelegates::BeginPIE;
+FEditorDelegates::FOnPIEEvent							FEditorDelegates::PrePIEEnded;
+FEditorDelegates::FOnPIEEvent							FEditorDelegates::PostPIEStarted;
 FEditorDelegates::FOnPIEEvent							FEditorDelegates::EndPIE;
 FEditorDelegates::FOnPIEEvent							FEditorDelegates::PausePIE;
 FEditorDelegates::FOnPIEEvent							FEditorDelegates::ResumePIE;
 FEditorDelegates::FOnPIEEvent							FEditorDelegates::SingleStepPIE;
+FEditorDelegates::FOnPIEEvent							FEditorDelegates::OnPreSwitchBeginPIEAndSIE;
 FEditorDelegates::FOnPIEEvent							FEditorDelegates::OnSwitchBeginPIEAndSIE;
 FSimpleMulticastDelegate								FEditorDelegates::PropertySelectionChange;
 FSimpleMulticastDelegate								FEditorDelegates::PostLandscapeLayerUpdated;
diff --git a/Engine/Source/Editor/UnrealEd/Private/EditorEngine.cpp b/Engine/Source/Editor/UnrealEd/Private/EditorEngine.cpp
--- a/Engine/Source/Editor/UnrealEd/Private/EditorEngine.cpp
+++ b/Engine/Source/Editor/UnrealEd/Private/EditorEngine.cpp
@@ -134,7 +134,7 @@
 #include "Components/SkeletalMeshComponent.h"
 #include "EngineModule.h"
 
-#include "EditorWorldManager.h"
+#include "EditorWorldExtension.h"
 
 #if PLATFORM_WINDOWS
 	#include "WindowsHWrapper.h"
@@ -317,6 +317,8 @@ UEditorEngine::UEditorEngine(const FObjectInitializer& ObjectInitializer)
 	NumOnlinePIEInstances = 0;
 	DefaultWorldFeatureLevel = GMaxRHIFeatureLevel;
 
+	EditorWorldExtensionsManager = nullptr;
+
 #if !UE_BUILD_SHIPPING
 	if (!AutomationCommon::OnEditorAutomationMapLoadDelegate().IsBound())
 	{
@@ -613,7 +615,7 @@ void UEditorEngine::InitEditor(IEngineLoop* InEngineLoop)
 	TimerManager = MakeShareable(new FTimerManager());
 
 	// create the editor world manager
-	EditorWorldManager = MakeShareable(new FEditorWorldManager());
+	EditorWorldExtensionsManager = NewObject<UEditorWorldExtensionManager>();
 
 	// Settings.
 	FBSPOps::GFastRebuild = 0;
@@ -1703,8 +1705,8 @@ void UEditorEngine::Tick( float DeltaSeconds, bool bIdleMode )
 		}
 	}
 
-	// Updates the ViewportWorldInteraction
-	EditorWorldManager->Tick( DeltaSeconds );
+	// Updates all the extensions for all the editor worlds
+	EditorWorldExtensionsManager->Tick(DeltaSeconds);
 
 	bool bIsMouseOverAnyLevelViewport = false;
 
@@ -2265,7 +2267,7 @@ void UEditorEngine::PlayEditorSound( const FString& SoundAssetName )
 void UEditorEngine::PlayEditorSound( USoundBase* InSound )
 {
 	// Only play sounds if the user has that feature enabled
-	if (!GIsSavingPackage && IsInGameThread() && GetDefault<ULevelEditorMiscSettings>()->bEnableEditorSounds)
+	if (!GIsSavingPackage && CanPlayEditorSound())
 	{
 		if (InSound != nullptr)
 		{
@@ -2274,6 +2276,11 @@ void UEditorEngine::PlayEditorSound( USoundBase* InSound )
 	}
 }
 
+bool UEditorEngine::CanPlayEditorSound() const
+{
+	return IsInGameThread() && GetDefault<ULevelEditorMiscSettings>()->bEnableEditorSounds;
+}
+
 void UEditorEngine::ClearPreviewComponents()
 {
 	if( PreviewAudioComponent )
@@ -7007,4 +7014,4 @@ void UEditorEngine::AutomationLoadMap(const FString& MapName, FString* OutError)
 	return;
 }
 
-#undef LOCTEXT_NAMESPACE 
+#undef LOCTEXT_NAMESPACE 
\ No newline at end of file
diff --git a/Engine/Source/Editor/UnrealEd/Private/EditorModeManager.cpp b/Engine/Source/Editor/UnrealEd/Private/EditorModeManager.cpp
--- a/Engine/Source/Editor/UnrealEd/Private/EditorModeManager.cpp
+++ b/Engine/Source/Editor/UnrealEd/Private/EditorModeManager.cpp
@@ -18,6 +18,11 @@
 #include "Framework/Notifications/NotificationManager.h"
 #include "Widgets/Notifications/SNotificationList.h"
 #include "Engine/LevelStreaming.h"
+#include "EditorWorldExtension.h"
+#include "ViewportWorldInteraction.h"
+#include "Editor/EditorEngine.h"
+#include "UnrealEdGlobals.h"
+#include "Editor/UnrealEdEngine.h"
 
 /*------------------------------------------------------------------------------
 	FEditorModeTools.
@@ -217,7 +222,18 @@ void FEditorModeTools::SetPivotLocation( const FVector& Location, const bool bIn
 
 ECoordSystem FEditorModeTools::GetCoordSystem(bool bGetRawValue)
 {
-	if (!bGetRawValue && (GetWidgetMode() == FWidget::WM_Scale))
+	bool bAligningToActors = false;
+	if (GEditor->GetEditorWorldExtensionsManager() != nullptr)
+	{
+
+		UViewportWorldInteraction* ViewportWorldInteraction = Cast<UViewportWorldInteraction>(GEditor->GetEditorWorldExtensionsManager()->GetEditorWorldExtensions(GetWorld())->FindExtension(UViewportWorldInteraction::StaticClass()));
+		if (ViewportWorldInteraction != nullptr && ViewportWorldInteraction->AreAligningToActors() == true)
+		{
+			bAligningToActors = true;
+		}
+	}
+	if (!bGetRawValue && 
+		((GetWidgetMode() == FWidget::WM_Scale) || bAligningToActors))
 	{
 		return COORD_Local;
 	}
@@ -229,6 +245,20 @@ ECoordSystem FEditorModeTools::GetCoordSystem(bool bGetRawValue)
 
 void FEditorModeTools::SetCoordSystem(ECoordSystem NewCoordSystem)
 {
+	// If we are trying to enter world space but are aligning to actors, turn off aligning to actors
+	if (GEditor->GetEditorWorldExtensionsManager() != nullptr &&
+		NewCoordSystem == COORD_World)
+	{
+		UViewportWorldInteraction* ViewportWorldInteraction = Cast<UViewportWorldInteraction>(GEditor->GetEditorWorldExtensionsManager()->GetEditorWorldExtensions(GetWorld())->FindExtension(UViewportWorldInteraction::StaticClass()));
+		if (ViewportWorldInteraction != nullptr && ViewportWorldInteraction->AreAligningToActors() == true)
+		{
+			if (ViewportWorldInteraction->HasCandidatesSelected())
+			{
+				ViewportWorldInteraction->SetSelectionAsCandidates();
+			}
+			GUnrealEd->Exec(GetWorld(), TEXT("VI.EnableGuides 0"));
+		}
+	}
 	CoordSystem = NewCoordSystem;
 }
 
diff --git a/Engine/Source/Editor/UnrealEd/Private/EditorViewportClient.cpp b/Engine/Source/Editor/UnrealEd/Private/EditorViewportClient.cpp
--- a/Engine/Source/Editor/UnrealEd/Private/EditorViewportClient.cpp
+++ b/Engine/Source/Editor/UnrealEd/Private/EditorViewportClient.cpp
@@ -43,6 +43,8 @@
 #include "ComponentRecreateRenderStateContext.h"
 #include "EditorBuildUtils.h"
 #include "AudioDevice.h"
+#include "EditorWorldExtension.h"
+#include "ViewportWorldInteraction.h"
 
 #define LOCTEXT_NAMESPACE "EditorViewportClient"
 
@@ -1423,8 +1425,16 @@ EMouseCursor::Type FEditorViewportClient::GetCursor(FViewport* InViewport,int32
 			}
 		}
 	}
-
-
+	
+	// Allow the viewport interaction to override any previously set mouse cursor
+	UViewportWorldInteraction* WorldInteraction = Cast<UViewportWorldInteraction>(GEditor->GetEditorWorldExtensionsManager()->GetEditorWorldExtensions(GetWorld())->FindExtension(UViewportWorldInteraction::StaticClass()));
+	if (WorldInteraction != nullptr && WorldInteraction->ShouldSuppressExistingCursor())
+	{
+			MouseCursor = EMouseCursor::None;
+			RequiredCursorVisibiltyAndAppearance.bHardwareCursorVisible = false;
+			RequiredCursorVisibiltyAndAppearance.bSoftwareCursorVisible = false;
+			UpdateRequiredCursorVisibility();
+	}
 
 	CachedMouseX = X;
 	CachedMouseY = Y;
@@ -2276,6 +2286,12 @@ bool FEditorViewportClient::InputKey(FViewport* InViewport, int32 ControllerId,
 		return true;
 	}
 
+	UEditorWorldExtensionCollection& EditorWorldExtensionCollection = *GEditor->GetEditorWorldExtensionsManager()->GetEditorWorldExtensions( GetWorld() );
+	if( EditorWorldExtensionCollection.InputKey(this, Viewport, Key, Event))
+	{
+		return true;
+	}
+
 	FInputEventState InputState(InViewport, Key, Event);
 
 	bool bHandled = false;
