Branch snapshot for CL 3327560

[CL 3327560 in Dev-Geometry branch]

diff --git a/Engine/Config/BaseEditorPerProjectUserSettings.ini b/Engine/Config/BaseEditorPerProjectUserSettings.ini
--- a/Engine/Config/BaseEditorPerProjectUserSettings.ini
+++ b/Engine/Config/BaseEditorPerProjectUserSettings.ini
@@ -153,12 +153,8 @@ bBreakOnExceptions=False
 bToolbarCustomization=False
 bBehaviorTreeEditor=False
 bEnableFindAndReplaceReferences =False
-; Enables a preview of the Unreal Editor in VR.  This adds a new tool bar button that allows you to toggle into "VR Mode" instantly.  This feature is still in development, but your feedback is appreciated!
-bEnableVREditing=false
-; If true, wearing a Vive or Oculus Rift headset will automatically enter VR Editing mode if Enable VR Editing is true. 
-bEnableAutoVREditMode=true
-; If true, the VR Tutorial window will always be shown at startup
-bAlwaysShowVRTutorial=true
+; If true, wearing a Vive or Oculus Rift headset will automatically enter VR Editing mode 
+bEnableAutoVREditMode=false
 
 [/Script/UnrealEd.EditorLoadingSavingSettings]
 ; True if we should automatically load a default level at start up
diff --git a/Engine/Plugins/MovieScene/LevelSequenceEditor/Source/LevelSequenceEditor/Private/LevelSequenceEditorToolkit.cpp b/Engine/Plugins/MovieScene/LevelSequenceEditor/Source/LevelSequenceEditor/Private/LevelSequenceEditorToolkit.cpp
--- a/Engine/Plugins/MovieScene/LevelSequenceEditor/Source/LevelSequenceEditor/Private/LevelSequenceEditorToolkit.cpp
+++ b/Engine/Plugins/MovieScene/LevelSequenceEditor/Source/LevelSequenceEditor/Private/LevelSequenceEditorToolkit.cpp
@@ -48,9 +48,14 @@
 #include "UnrealEdMisc.h"
 
 // @todo sequencer: hack: setting defaults for transform tracks
+
 #include "Sections/MovieScene3DTransformSection.h"
 #include "Tracks/MovieScene3DTransformTrack.h"
 
+// To override Sequencer editor behavior for VR Editor 
+#include "VREditorMode.h"
+
+
 #define LOCTEXT_NAMESPACE "LevelSequenceEditor"
 
 
@@ -157,6 +162,7 @@ FLevelSequenceEditorToolkit::FLevelSequenceEditorToolkit(const TSharedRef<ISlate
 		FAssetEditorExtender::CreateRaw(this, &FLevelSequenceEditorToolkit::HandleMenuExtensibilityGetExtender));
 	SequencerExtenderHandle = SequencerModule.GetAddTrackMenuExtensibilityManager()->GetExtenderDelegates()[NewIndex].GetHandle();
 
+
 	OpenToolkits.Add(this);
 }
 
@@ -271,6 +277,22 @@ void FLevelSequenceEditorToolkit::Initialize(const EToolkitMode::Type Mode, cons
 	LevelSequenceEditorModule.OnMasterSequenceCreated().AddRaw(this, &FLevelSequenceEditorToolkit::HandleMasterSequenceCreated);
 
 	FLevelSequenceEditorToolkit::OnOpened().Broadcast(*this);
+
+	{
+		UWorld* World = CastChecked<UWorld>( GetLevelSequenceEditorPlaybackContext() );
+		UVREditorMode* VRMode = Cast<UVREditorMode>( GEditor->GetEditorWorldExtensionsManager()->GetEditorWorldExtensions( World )->FindExtension( UVREditorMode::StaticClass() ) );
+		if (VRMode != nullptr)
+		{
+			VRMode->OnVREditingModeExit_Handler.BindSP(this, &FLevelSequenceEditorToolkit::HandleVREditorModeExit);
+			USequencerSettings& SavedSequencerSettings = *Sequencer->GetSequencerSettings();
+			VRMode->SaveSequencerSettings(Sequencer->GetKeyAllEnabled(), Sequencer->GetAutoKeyMode(), SavedSequencerSettings);
+			// Override currently set autokey behavior to always autokey all
+			Sequencer->SetAutoKeyMode(EAutoKeyMode::KeyAll);
+			Sequencer->SetKeyAllEnabled(true);
+			// Tell the VR Editor mode that Sequencer has refreshed
+			VRMode->RefreshVREditorSequencer(Sequencer.Get());
+		}
+	}
 }
 
 
@@ -702,6 +724,17 @@ void FLevelSequenceEditorToolkit::HandleActorAddedToSequencer(AActor* Actor, con
 }
 
 
+void FLevelSequenceEditorToolkit::HandleVREditorModeExit()
+{
+	UWorld* World = Cast<UWorld>(GetLevelSequenceEditorPlaybackContext());
+	UVREditorMode* VRMode = CastChecked<UVREditorMode>( GEditor->GetEditorWorldExtensionsManager()->GetEditorWorldExtensions( World )->FindExtension( UVREditorMode::StaticClass() ) );
+
+	// Reset sequencer settings
+	Sequencer->SetAutoKeyMode(VRMode->GetSavedEditorState().AutoKeyMode);
+	Sequencer->SetKeyAllEnabled(VRMode->GetSavedEditorState().bKeyAllEnabled);
+	VRMode->OnVREditingModeExit_Handler.Unbind();
+}
+
 void FLevelSequenceEditorToolkit::HandleMapChanged(class UWorld* NewWorld, EMapChangeType MapChangeType)
 {
 	// @todo sequencer: We should only wipe/respawn puppets that are affected by the world that is being changed! (multi-UWorld support)
@@ -774,7 +807,6 @@ void FLevelSequenceEditorToolkit::AddShot(UMovieSceneCinematicShotTrack* ShotTra
 		{
 			CameraGuid = GetSequencer()->CreateBinding(*NewCamera, NewCamera->GetActorLabel());
 		}
-
 		NewCamera->SetActorLocation( GCurrentLevelEditingViewportClient->GetViewLocation(), false );
 		NewCamera->SetActorRotation( GCurrentLevelEditingViewportClient->GetViewRotation() );
 		//pNewCamera->CameraComponent->FieldOfView = ViewportClient->ViewFOV; //@todo set the focal length from this field of view
@@ -935,6 +967,13 @@ void FLevelSequenceEditorToolkit::HandleTrackMenuExtensionAddTrack(FMenuBuilder&
 
 bool FLevelSequenceEditorToolkit::OnRequestClose()
 {
+	UWorld* World = CastChecked<UWorld>(GetLevelSequenceEditorPlaybackContext());
+	UVREditorMode* VRMode = Cast<UVREditorMode>(GEditor->GetEditorWorldExtensionsManager()->GetEditorWorldExtensions(World)->FindExtension(UVREditorMode::StaticClass()));
+	if (VRMode != nullptr)
+	{
+		// Null out the VR Mode's sequencer pointer
+		VRMode->RefreshVREditorSequencer(nullptr);
+	}
 	OpenToolkits.Remove(this);
 	OnClosedEvent.Broadcast();
 	return true;
diff --git a/Engine/Plugins/Runtime/OculusRift/Source/OculusRift/Private/OculusRiftHMD.cpp b/Engine/Plugins/Runtime/OculusRift/Source/OculusRift/Private/OculusRiftHMD.cpp
--- a/Engine/Plugins/Runtime/OculusRift/Source/OculusRift/Private/OculusRiftHMD.cpp
+++ b/Engine/Plugins/Runtime/OculusRift/Source/OculusRift/Private/OculusRiftHMD.cpp
@@ -26,6 +26,9 @@
 
 static TAutoConsoleVariable<int32> CStartInVRVar(TEXT("vr.bStartInVR"), 0, TEXT("Start in VR flag"));
 static TAutoConsoleVariable<int32> CGracefulExitVar(TEXT("vr.bExitGracefully"), 0, TEXT("Exit gracefully when forced by Universal Menu."));
+static TAutoConsoleVariable<float> CRightEyeResolutionScaling( TEXT( "rift.RightEyeResolutionScaling" ), 1.0f, TEXT( "Scales the resolution of the right eye only" ) );
+static TAutoConsoleVariable<int32> CLeftEyeOverscan( TEXT( "rift.LeftEyeOverscan" ), 0, TEXT( "Whether to render the left eye at a larger widescreen resolution, then crop a column from the center for the HMD.  This is a presentation mode feature." ) );
+static TAutoConsoleVariable<float> CLeftEyeOverscanAspect( TEXT( "rift.LeftEyeOverscanAspect" ), 16.0f / 9.0f, TEXT( "Aspect ratio when rift.LeftEyeOverscan is enabled" ) );
 
 //---------------------------------------------------
 // Oculus Rift Plugin Implementation
@@ -580,17 +583,17 @@ bool FOculusRiftHMD::OnStartGameFrame( FWorldContext& InWorldContext )
 		const OVR::Vector3f newRight(CurrentSettings->EyeRenderDesc[1].HmdToEyeOffset);
 		const OVR::Vector3f prevLeft(MasterSettings->EyeRenderDesc[0].HmdToEyeOffset);
 		const OVR::Vector3f prevRight(MasterSettings->EyeRenderDesc[1].HmdToEyeOffset);
-		if (newLeft != prevLeft || newRight != prevRight)
+		// for debugging purposes only: overriding IPD
+		if( CurrentSettings->Flags.bOverrideIPD )
 		{
-			const float newIAD = newRight.Distance(newLeft);
-			UE_LOG(LogHMD, Log, TEXT("IAD has changed, new IAD is %.4f meters"), newIAD);
+			check( CurrentSettings->InterpupillaryDistance >= 0 );
+			CurrentSettings->EyeRenderDesc[ 0 ].HmdToEyeOffset.x = -CurrentSettings->InterpupillaryDistance * 0.5f;
+			CurrentSettings->EyeRenderDesc[ 1 ].HmdToEyeOffset.x = CurrentSettings->InterpupillaryDistance * 0.5f;
 		}
-		// for debugging purposes only: overriding IPD
-		if (CurrentSettings->Flags.bOverrideIPD)
+		else if (newLeft != prevLeft || newRight != prevRight)
 		{
-			check(CurrentSettings->InterpupillaryDistance >= 0);
-			CurrentSettings->EyeRenderDesc[0].HmdToEyeOffset.x = -CurrentSettings->InterpupillaryDistance * 0.5f;
-			CurrentSettings->EyeRenderDesc[1].HmdToEyeOffset.x = CurrentSettings->InterpupillaryDistance * 0.5f;
+			const float newIAD = newRight.Distance(newLeft);
+			UE_LOG(LogHMD, Log, TEXT("IAD has changed, new IAD is %.4f meters"), newIAD);
 		}
 #endif // #if !UE_BUILD_SHIPPING
 		// Save EyeRenderDesc in main settings.
@@ -1752,7 +1755,9 @@ void FOculusRiftHMD::SetupViewFamily(FSceneViewFamily& InViewFamily)
 	auto frame = GetFrame();
 	check(frame);
 
-	InViewFamily.EngineShowFlags.MotionBlur = 0;
+	extern RENDERER_API TAutoConsoleVariable<int32> CVarAllowMotionBlurInVR;
+
+	InViewFamily.EngineShowFlags.MotionBlur = CVarAllowMotionBlurInVR->GetInt() != 0;
 	InViewFamily.EngineShowFlags.HMDDistortion = false;
 	InViewFamily.EngineShowFlags.StereoRendering = IsStereoEnabled();
 	if (frame->Settings->Flags.bPauseRendering)
@@ -2137,10 +2142,6 @@ void FOculusRiftHMD::UpdateHmdRenderInfo()
 		// Calc FOV, symmetrical, for each eye. 
 		CurrentSettings->EyeFov[0] = HmdDesc.DefaultEyeFov[0];
 		CurrentSettings->EyeFov[1] = HmdDesc.DefaultEyeFov[1];
-
-		// Calc FOV in radians
-		CurrentSettings->VFOVInRadians = GetVerticalFovRadians(CurrentSettings->EyeFov[0], CurrentSettings->EyeFov[1]);
-		CurrentSettings->HFOVInRadians = GetHorizontalFovRadians(CurrentSettings->EyeFov[0], CurrentSettings->EyeFov[1]);
 	}
 
 	const ovrSizei recommenedTex0Size = ovr_GetFovTextureSize(OvrSession, ovrEye_Left, CurrentSettings->EyeFov[0], 1.0f);
@@ -2153,6 +2154,27 @@ void FOculusRiftHMD::UpdateHmdRenderInfo()
 	CurrentSettings->IdealScreenPercentage = FMath::Max(float(idealRenderTargetSize.w) / float(HmdDesc.Resolution.w) * 100.f,
 														float(idealRenderTargetSize.h) / float(HmdDesc.Resolution.h) * 100.f);
 
+	if (!CurrentSettings->Flags.bOverrideFOV)
+	{
+		if( CLeftEyeOverscan->GetInt() != 0 )
+		{
+			// Compute a new FOV for the left eye that will "overscan" the HMD's viewport to allow for a
+			// much wider mirror image.  The user of the HMD will see the horizontally cropped section of
+			// view that would normally be rendered without the overscan, but viewers of the mirror display
+			// will see the entire view.
+			const ovrSizei OriginalEyeRenderTargetSize = ovr_GetFovTextureSize( OvrSession, ovrEye_Left, CurrentSettings->EyeFov[ 0 ], 1.0f );
+			const float OriginalEyeAspect = (float)OriginalEyeRenderTargetSize.w / (float)OriginalEyeRenderTargetSize.h;
+			const float NewEyeAspect = CLeftEyeOverscanAspect->GetFloat();
+
+			CurrentSettings->EyeFov[ 0 ].LeftTan *= NewEyeAspect / OriginalEyeAspect;
+			CurrentSettings->EyeFov[ 0 ].RightTan *= NewEyeAspect / OriginalEyeAspect;
+		}
+
+		// Calc FOV in radians
+		CurrentSettings->VFOVInRadians = GetVerticalFovRadians(CurrentSettings->EyeFov[0], CurrentSettings->EyeFov[1]);
+		CurrentSettings->HFOVInRadians = GetHorizontalFovRadians(CurrentSettings->EyeFov[0], CurrentSettings->EyeFov[1]);
+	}
+
 	// Override eye distance by the value from HMDInfo (stored in Profile).
 	if (!CurrentSettings->Flags.bOverrideIPD)
 	{
@@ -2242,7 +2264,13 @@ void FOculusRiftHMD::UpdateStereoRenderingParams()
 
 		for(int eyeIndex = 0; eyeIndex < ovrEye_Count; eyeIndex++)
 		{
-			vpSize[eyeIndex] = ovr_GetFovTextureSize(OvrSession, (ovrEyeType) eyeIndex, CurrentSettings->EyeFov[eyeIndex], pixelDensity);
+			float eyePixelDensity = pixelDensity;
+			if( eyeIndex == ovrEye_Right )
+			{
+				eyePixelDensity *= CRightEyeResolutionScaling->GetFloat();
+			}
+
+			vpSize[eyeIndex] = ovr_GetFovTextureSize(OvrSession, (ovrEyeType) eyeIndex, CurrentSettings->EyeFov[eyeIndex], eyePixelDensity);
 			vpSize[eyeIndex].w = FMath::Max(vpSize[eyeIndex].w, 1);
 			vpSize[eyeIndex].h = FMath::Max(vpSize[eyeIndex].h, 1);
 
@@ -2512,32 +2540,35 @@ void FOculusRiftHMD::OnBeginPlay(FWorldContext& InWorldContext)
 	CachedViewportWidget.Reset();
 	CachedWindow.Reset();
 
-#if WITH_EDITOR
-	// @TODO: add more values here.
-	// This call make sense when 'Play' is used from the Editor;
-	if (GIsEditor)
+	if( !GEnableVREditorHacks )
 	{
-		Settings->PositionOffset = FVector::ZeroVector;
-		Settings->BaseOrientation = FQuat::Identity;
-		Settings->BaseOffset = FVector::ZeroVector;
-		Settings->WorldToMetersScale = InWorldContext.World()->GetWorldSettings()->WorldToMeters;
-		//Settings->Flags.bWorldToMetersOverride = false;
-		InitDevice();
+#if WITH_EDITOR
+		// @TODO: add more values here.
+		// This call make sense when 'Play' is used from the Editor;
+		if( GIsEditor )
+		{
+			Settings->PositionOffset = FVector::ZeroVector;
+			Settings->BaseOrientation = FQuat::Identity;
+			Settings->BaseOffset = FVector::ZeroVector;
+			Settings->WorldToMetersScale = InWorldContext.World()->GetWorldSettings()->WorldToMeters;
+			//Settings->Flags.bWorldToMetersOverride = false;
+			InitDevice();
 
-		FApp::SetUseVRFocus(true);
-		FApp::SetHasVRFocus(true);
-		OnStartGameFrame(InWorldContext);
-	}
-	else
+			FApp::SetUseVRFocus( true );
+			FApp::SetHasVRFocus( true );
+			OnStartGameFrame( InWorldContext );
+		}
+		else
 #endif
-	{
-		MakeSureValidFrameExists(InWorldContext.World()->GetWorldSettings());
+		{
+			MakeSureValidFrameExists( InWorldContext.World()->GetWorldSettings() );
+		}
 	}
 }
 
 void FOculusRiftHMD::OnEndPlay(FWorldContext& InWorldContext)
 {
-	if (GIsEditor)
+	if (GIsEditor && !GEnableVREditorHacks)
 	{
 		// @todo vreditor: If we add support for starting PIE while in VR Editor, we don't want to kill stereo mode when exiting PIE
 		EnableStereo(false);
diff --git a/Engine/Shaders/MaterialTemplate.usf b/Engine/Shaders/MaterialTemplate.usf
--- a/Engine/Shaders/MaterialTemplate.usf
+++ b/Engine/Shaders/MaterialTemplate.usf
@@ -207,6 +207,13 @@ struct FMaterialPixelParameters
 	/** Used in mobile custom pp material to preserve original SceneColor Alpha */
 	half BackupSceneColorAlpha;
 #endif
+
+#if COMPILER_HLSL
+	// Workaround for "error X3067: 'GetObjectWorldPosition': ambiguous function call"
+	// Which happens when FMaterialPixelParameters and FMaterialVertexParameters have the same number of floats with the HLSL compiler ver 9.29.952.3111
+	// Function overload resolution appears to identify types based on how many floats / ints / etc they contain
+	uint Dummy;
+#endif
 };
 
 // @todo compat hack
diff --git a/Engine/Source/Editor/ContentBrowser/Private/SAssetView.cpp b/Engine/Source/Editor/ContentBrowser/Private/SAssetView.cpp
--- a/Engine/Source/Editor/ContentBrowser/Private/SAssetView.cpp
+++ b/Engine/Source/Editor/ContentBrowser/Private/SAssetView.cpp
@@ -159,12 +159,6 @@ void SAssetView::Construct( const FArguments& InArgs )
 	TileViewThumbnailSize = 128;
 	TileViewThumbnailPadding = 5;
 
-	const bool bIsVREditorDemo = FParse::Param( FCommandLine::Get(), TEXT( "VREditorDemo" ) );	// @todo vreditor: Remove this when no longer needed
-	if( bIsVREditorDemo )
-	{
-		TileViewThumbnailPadding = 0;
-	}
-
 	TileViewNameHeight = 36;
 	ThumbnailScaleSliderValue = InArgs._ThumbnailScale; 
 
@@ -1051,12 +1045,6 @@ void SAssetView::CalculateFillScale( const FGeometry& AllottedGeometry )
 	{
 		float ItemWidth = GetTileViewItemBaseWidth();
 
- 		const bool bIsVREditorDemo = FParse::Param( FCommandLine::Get(), TEXT( "VREditorDemo" ) );	// @todo vreditor: Remove this when no longer needed
-		if( bIsVREditorDemo )
-		{
-			ItemWidth /= AllottedGeometry.Scale;
-		}
-
 		// Scrollbars are 16, but we add 1 to deal with half pixels.
 		const float ScrollbarWidth = 16 + 1;
 		float TotalWidth = AllottedGeometry.Size.X - ( ScrollbarWidth / AllottedGeometry.Scale );
@@ -3760,8 +3748,12 @@ FReply SAssetView::OnDraggingAssetItem( const FGeometry& MyGeometry, const FPoin
 			if( InAssetData.Num() > 0 )
 			{
 				UActorFactory* FactoryToUse = nullptr;
-				FEditorDelegates::OnAssetDragStarted.Broadcast( InAssetData, FactoryToUse );
-				if( MouseEvent.IsMouseButtonDown( EKeys::LeftMouseButton ) )
+				if( FEditorDelegates::OnAssetDragStarted.IsBound() )
+				{
+					FEditorDelegates::OnAssetDragStarted.Broadcast( InAssetData, FactoryToUse );
+					return FReply::Handled();
+				}
+				else if( MouseEvent.IsMouseButtonDown( EKeys::LeftMouseButton ) )
 				{
 					return FReply::Handled().BeginDragDrop( FAssetDragDropOp::New( InAssetData ) );
 				}
diff --git a/Engine/Source/Editor/EditorStyle/Private/SlateEditorStyle.cpp b/Engine/Source/Editor/EditorStyle/Private/SlateEditorStyle.cpp
--- a/Engine/Source/Editor/EditorStyle/Private/SlateEditorStyle.cpp
+++ b/Engine/Source/Editor/EditorStyle/Private/SlateEditorStyle.cpp
@@ -4624,10 +4624,10 @@ void FSlateEditorStyle::FStyle::SetupLevelEditorStyle()
 		Set( "LevelEditor.FoliageMode.Selected", new IMAGE_BRUSH( "Icons/icon_Mode_Foliage_selected_40x", Icon40x40 ) );
 		Set( "LevelEditor.FoliageMode.Selected.Small", new IMAGE_BRUSH( "Icons/icon_Mode_Foliage_selected_40x", Icon20x20 ) );
 
-		Set( "LevelEditor.BspMode", new IMAGE_BRUSH( "Icons/icon_Mode_GeoEdit_40px", Icon40x40 ) );
-		Set( "LevelEditor.BspMode.Small", new IMAGE_BRUSH( "Icons/icon_Mode_GeoEdit_40px", Icon20x20 ) );
-		Set( "LevelEditor.BspMode.Selected", new IMAGE_BRUSH( "Icons/icon_Mode_GeoEdit-a_40px", Icon40x40 ) );
-		Set( "LevelEditor.BspMode.Selected.Small", new IMAGE_BRUSH( "Icons/icon_Mode_GeoEdit-a_40px", Icon20x20 ) );
+		Set( "LevelEditor.BspMode", new IMAGE_BRUSH( "Icons/icon_Mode_MeshEditor_40px", Icon40x40 ) );
+		Set( "LevelEditor.BspMode.Small", new IMAGE_BRUSH( "Icons/icon_Mode_MeshEditor_40px", Icon20x20 ) );
+		Set( "LevelEditor.BspMode.Selected", new IMAGE_BRUSH( "Icons/icon_Mode_MeshEditor_40px-a_40px", Icon40x40 ) );
+		Set( "LevelEditor.BspMode.Selected.Small", new IMAGE_BRUSH( "Icons/icon_Mode_MeshEditor_40px-a_40px", Icon20x20 ) );
 
 		Set( "LevelEditor.WorldProperties", new IMAGE_BRUSH( "Icons/icon_worldscript_40x", Icon40x40 ) );
 		Set( "LevelEditor.WorldProperties.Small", new IMAGE_BRUSH( "Icons/icon_worldscript_40x", Icon20x20 ) );
@@ -4644,8 +4644,8 @@ void FSlateEditorStyle::FStyle::SetupLevelEditorStyle()
 		Set( "LevelEditor.EditMatinee", new IMAGE_BRUSH( "Icons/icon_matinee_40x", Icon40x40 ) );
 		Set( "LevelEditor.EditMatinee.Small", new IMAGE_BRUSH( "Icons/icon_matinee_40x", Icon20x20 ) );
 
-		Set( "LevelEditor.ToggleVR", new IMAGE_BRUSH( "Icons/icon_ToggleVREditor_40x", Icon40x40 ) );
-		Set( "LevelEditor.ToggleVR.Small", new IMAGE_BRUSH( "Icons/icon_ToggleVREditor_40x", Icon20x20 ) );
+		Set( "LevelEditor.ToggleVR", new IMAGE_BRUSH( "Icons/VREditor/VR_Editor_Toolbar_Icon", Icon40x40 ) );
+		Set( "LevelEditor.ToggleVR.Small", new IMAGE_BRUSH( "Icons/VREditor/VR_Editor_Toolbar_Icon_Small", Icon20x20 ) );
 
 		Set( "MergeActors.MeshMergingTool", new IMAGE_BRUSH( "Icons/icon_MergeActors_MeshMerging_40x", Icon40x40 ) );
 		Set( "MergeActors.MeshProxyTool", new IMAGE_BRUSH( "Icons/icon_MergeActors_MeshProxy_40x", Icon40x40 ) );
@@ -5911,17 +5911,9 @@ void FSlateEditorStyle::FStyle::SetupContentBrowserStyle()
 		Set( "ContentBrowser.SourceTitleFont", TTF_CORE_FONT( "Fonts/Roboto-Regular", 12 ) );
 
 		// @todo vreditor urgent: Increase the size of Content Browser fonts while in VR
-		const bool bIsVREditorDemo = FParse::Param( FCommandLine::Get(), TEXT( "VREditorDemo" ) );	// @todo vreditor: Remove this when no longer needed
-		if( bIsVREditorDemo )
-		{
-			Set( "ContentBrowser.SourceListItemFont", TTF_CORE_FONT( "Fonts/Roboto-Regular", 20 ) );
-			Set( "ContentBrowser.SourceTreeItemFont", TTF_CORE_FONT( "Fonts/Roboto-Regular", 20 ) );
-		}
-		else
-		{
-			Set( "ContentBrowser.SourceListItemFont", TTF_CORE_FONT( "Fonts/Roboto-Regular", 10 ) );
-			Set( "ContentBrowser.SourceTreeItemFont", TTF_CORE_FONT( "Fonts/Roboto-Regular", 10 ) );
-		}
+		Set("ContentBrowser.SourceListItemFont", TTF_CORE_FONT("Fonts/Roboto-Regular", 10));
+		Set("ContentBrowser.SourceTreeItemFont", TTF_CORE_FONT("Fonts/Roboto-Regular", 10));
+
 		Set( "ContentBrowser.SourceTreeRootItemFont", TTF_CORE_FONT( "Fonts/Roboto-Regular", 12 ) );
 		Set( "ContentBrowser.AssetTreeFolderClosed", new IMAGE_BRUSH( "Icons/FolderClosed", FVector2D(18, 16) ) );
 		Set( "ContentBrowser.BreadcrumbPathPickerFolder", new IMAGE_BRUSH( "Icons/FolderClosed", FVector2D(18, 16), FLinearColor::Gray ) );
diff --git a/Engine/Source/Editor/FoliageEdit/Private/FoliageEdMode.cpp b/Engine/Source/Editor/FoliageEdit/Private/FoliageEdMode.cpp
--- a/Engine/Source/Editor/FoliageEdit/Private/FoliageEdMode.cpp
+++ b/Engine/Source/Editor/FoliageEdit/Private/FoliageEdMode.cpp
@@ -14,6 +14,7 @@
 #include "Framework/Application/SlateApplication.h"
 #include "Framework/Commands/UICommandList.h"
 #include "Materials/Material.h"
+#include "Materials/MaterialInstanceDynamic.h"
 #include "Engine/CollisionProfile.h"
 #include "Engine/StaticMeshActor.h"
 #include "Components/ModelComponent.h"
@@ -42,11 +43,10 @@
 #include "Components/BrushComponent.h"
 
 // VR Editor
-#include "IVREditorModule.h"
 #include "VREditorMode.h"
 #include "ViewportWorldInteraction.h"
 #include "VREditorInteractor.h"
-#include "EditorWorldManager.h"
+#include "EditorWorldExtension.h"
 
 
 #define LOCTEXT_NAMESPACE "FoliageEdMode"
@@ -54,6 +54,11 @@
 
 DEFINE_LOG_CATEGORY_STATIC(LogFoliage, Log, Warning);
 
+namespace VREd
+{
+	static FAutoConsoleVariable FoliageOpacity(TEXT("VREd.FoliageOpacity"), 0.02f, TEXT("The foliage brush opacity."));
+}
+
 //
 // FFoliageMeshUIInfo
 //
@@ -190,11 +195,13 @@ FEdModeFoliage::FEdModeFoliage()
 	, FoliageInteractor(nullptr)
 {
 	// Load resources and construct brush component
-	UMaterial* BrushMaterial = nullptr;
 	UStaticMesh* StaticMesh = nullptr;
 	if (!IsRunningCommandlet())
 	{
-		BrushMaterial = LoadObject<UMaterial>(nullptr, TEXT("/Engine/EditorLandscapeResources/FoliageBrushSphereMaterial.FoliageBrushSphereMaterial"), nullptr, LOAD_None, nullptr);
+		UMaterial* BrushMaterial = LoadObject<UMaterial>(nullptr, TEXT("/Engine/EditorLandscapeResources/FoliageBrushSphereMaterial.FoliageBrushSphereMaterial"), nullptr, LOAD_None, nullptr);
+		BrushMID = UMaterialInstanceDynamic::Create(BrushMaterial, GetTransientPackage());
+		check(BrushMID != nullptr);
+
 		StaticMesh = LoadObject<UStaticMesh>(nullptr, TEXT("/Engine/EngineMeshes/Sphere.Sphere"), nullptr, LOAD_None, nullptr);
 	}
 
@@ -202,12 +209,16 @@ FEdModeFoliage::FEdModeFoliage()
 	SphereBrushComponent->SetCollisionProfileName(UCollisionProfile::NoCollision_ProfileName);
 	SphereBrushComponent->SetCollisionObjectType(ECC_WorldDynamic);
 	SphereBrushComponent->SetStaticMesh(StaticMesh);
-	SphereBrushComponent->OverrideMaterials.Add(BrushMaterial);
+	SphereBrushComponent->SetMaterial(0, BrushMID);
 	SphereBrushComponent->SetAbsolute(true, true, true);
 	SphereBrushComponent->CastShadow = false;
 
 	bBrushTraceValid = false;
 	BrushLocation = FVector::ZeroVector;
+	
+	// Get the default opacity from the material.
+	FName OpacityParamName("OpacityAmount");
+	BrushMID->GetScalarParameterValue(OpacityParamName, DefaultBrushOpacity);
 
 	FFoliageEditCommands::Register();
 	UICommandList = MakeShareable(new FUICommandList);
@@ -357,7 +368,7 @@ void FEdModeFoliage::Enter()
 	NotifyNewCurrentLevel();
 
 	// Register for VR input events
-	UViewportWorldInteraction* ViewportWorldInteraction = GEditor->GetEditorWorldManager()->GetEditorWorldWrapper(GetWorld())->GetViewportWorldInteraction();
+	UViewportWorldInteraction* ViewportWorldInteraction = Cast<UViewportWorldInteraction>( GEditor->GetEditorWorldExtensionsManager()->GetEditorWorldExtensions( GetWorld() )->FindExtension( UViewportWorldInteraction::StaticClass() ) );
 	if (ViewportWorldInteraction != nullptr)
 	{
 		ViewportWorldInteraction->OnViewportInteractionInputAction().RemoveAll(this);
@@ -365,6 +376,11 @@ void FEdModeFoliage::Enter()
 
 		ViewportWorldInteraction->OnViewportInteractionHoverUpdate().RemoveAll(this);
 		ViewportWorldInteraction->OnViewportInteractionHoverUpdate().AddRaw(this, &FEdModeFoliage::OnVRHoverUpdate);
+
+		// Hide the VR transform gizmo while we're in foliage mode.  It sort of gets in the way of painting.
+		ViewportWorldInteraction->SetTransformGizmoVisible( false );
+
+		SetBrushOpacity(VREd::FoliageOpacity->GetFloat());
 	}
 }
 
@@ -372,14 +388,19 @@ void FEdModeFoliage::Enter()
 void FEdModeFoliage::Exit()
 {
 	// Unregister VR mode from event handlers
-	if (IVREditorModule::IsAvailable())
 	{
-		UViewportWorldInteraction* WorldInteraction = GEditor->GetEditorWorldManager()->GetEditorWorldWrapper(GetWorld())->GetViewportWorldInteraction();
-		if (WorldInteraction != nullptr)
+		UViewportWorldInteraction* ViewportWorldInteraction = Cast<UViewportWorldInteraction>( GEditor->GetEditorWorldExtensionsManager()->GetEditorWorldExtensions( GetWorld() )->FindExtension( UViewportWorldInteraction::StaticClass() ) );
+		if (ViewportWorldInteraction != nullptr)
 		{
-			WorldInteraction->OnViewportInteractionInputAction().RemoveAll(this);
-			WorldInteraction->OnViewportInteractionHoverUpdate().RemoveAll(this);
+			// Restore the transform gizmo visibility
+			ViewportWorldInteraction->SetTransformGizmoVisible( true );
+
+			ViewportWorldInteraction->OnViewportInteractionInputAction().RemoveAll(this);
+			ViewportWorldInteraction->OnViewportInteractionHoverUpdate().RemoveAll(this);
 			FoliageInteractor = nullptr;
+
+			// Reset the brush opacity to default.
+			SetBrushOpacity(DefaultBrushOpacity);
 		}
 	}
 
@@ -439,62 +460,65 @@ void FEdModeFoliage::Exit()
 	FEdMode::Exit();
 }
 
-void FEdModeFoliage::OnVRHoverUpdate(FEditorViewportClient& ViewportClient, UViewportInteractor* Interactor, FVector& HoverImpactPoint, bool& bWasHandled)
+void FEdModeFoliage::OnVRHoverUpdate(UViewportInteractor* Interactor, FVector& HoverImpactPoint, bool& bWasHandled)
 {
-	// Check if VR Editor is active
-	if (IVREditorModule::IsAvailable())
+	UVREditorMode* VREditorMode = Cast<UVREditorMode>( GEditor->GetEditorWorldExtensionsManager()->GetEditorWorldExtensions( GetWorld() )->FindExtension( UVREditorMode::StaticClass() ) );
+	if( VREditorMode != nullptr && VREditorMode->IsFullyInitialized() )
 	{
-		UVREditorMode* VREditorMode = GEditor->GetEditorWorldManager()->GetEditorWorldWrapper(ViewportClient.GetWorld())->GetVREditorMode();
-		if (VREditorMode != nullptr && VREditorMode->IsFullyInitialized())
+		// Check if we're hovering over UI. If so, stop painting so we don't display the preview brush sphere
+		if (FoliageInteractor && (FoliageInteractor->IsHoveringOverPriorityType() || FoliageInteractor->GetDraggingMode() != EViewportInteractionDraggingMode::Nothing))
 		{
-			// Check if we're hovering over UI. If so, stop painting so we don't display the preview brush sphere
-			if (FoliageInteractor && FoliageInteractor->IsHoveringOverPriorityType())
-			{
-				EndFoliageBrushTrace();
-				FoliageInteractor = nullptr;
-			}
-			// If there isn't currently a foliage interactor and we are hovering over something valid
-			else if (FoliageInteractor == nullptr && Interactor->GetHitResultFromLaserPointer().GetActor() != nullptr)
-			{
-				FoliageInteractor = Interactor;
-			}
-			// If we aren't hovering over something valid and the tool isn't active
-			else if (Interactor->GetHitResultFromLaserPointer().GetActor() == nullptr && !bToolActive)
-			{
-				FoliageInteractor = nullptr;
-			}
+			EndFoliageBrushTrace();
+			FoliageInteractor = nullptr;
+		}
+		// If there isn't currently a foliage interactor and we are hovering over something valid
+		else if (FoliageInteractor == nullptr && !Interactor->IsHoveringOverPriorityType() && Interactor->GetHitResultFromLaserPointer().GetActor() != nullptr)
+		{
+			FoliageInteractor = Interactor;
+		}
+		// If we aren't hovering over something valid and the tool isn't active
+		else if (Interactor->GetHitResultFromLaserPointer().GetActor() == nullptr && !bToolActive)
+		{
+			FoliageInteractor = nullptr;
+		}
 
-			// Skip other interactors if we are painting with one
-			if (FoliageInteractor && Interactor == FoliageInteractor)
+		// Skip other interactors if we are painting with one
+		if (FoliageInteractor && Interactor == FoliageInteractor)
+		{
+			// Go ahead and paint immediately
+			FVector LaserPointerStart, LaserPointerEnd;
+			if (FoliageInteractor->GetLaserPointer( /* Out */ LaserPointerStart, /* Out */ LaserPointerEnd))
 			{
-				// Go ahead and paint immediately
-				FVector LaserPointerStart, LaserPointerEnd;
-				if (FoliageInteractor->GetLaserPointer( /* Out */ LaserPointerStart, /* Out */ LaserPointerEnd))
-				{
-					const FVector LaserPointerDirection = (LaserPointerEnd - LaserPointerStart).GetSafeNormal();
+				const FVector LaserPointerDirection = (LaserPointerEnd - LaserPointerStart).GetSafeNormal();
+
+				FoliageBrushTrace(nullptr, LaserPointerStart, LaserPointerDirection);
 
-					FoliageBrushTrace(&ViewportClient, LaserPointerStart, LaserPointerDirection);
+				if( bBrushTraceValid )
+				{
+					HoverImpactPoint = BrushLocation;
+					bWasHandled = true;
 				}
 			}
-			
 		}
+			
+		const bool bBrushMeshVisible = !(FoliageInteractor == nullptr || Interactor->GetDraggingMode() != EViewportInteractionDraggingMode::Nothing);
+		SphereBrushComponent->SetVisibility(bBrushMeshVisible);
 	}
 }
 
 void FEdModeFoliage::OnVRAction(class FEditorViewportClient& ViewportClient, UViewportInteractor* Interactor, const FViewportActionKeyInput& Action, bool& bOutIsInputCaptured, bool& bWasHandled)
 {
-	UVREditorMode* VREditorMode = GEditor->GetEditorWorldManager()->GetEditorWorldWrapper(ViewportClient.GetWorld())->GetVREditorMode();
-	if (VREditorMode != nullptr && Interactor != nullptr)
+	UVREditorMode* VREditorMode = Cast<UVREditorMode>( GEditor->GetEditorWorldExtensionsManager()->GetEditorWorldExtensions( GetWorld() )->FindExtension( UVREditorMode::StaticClass() ) );
+	if (VREditorMode != nullptr && Interactor != nullptr && Interactor->GetDraggingMode() == EViewportInteractionDraggingMode::Nothing)
 	{
-		const UVREditorInteractor* VRInteractor = Cast<UVREditorInteractor>(Interactor);
-
-		// Only allow light press
-		Interactor->SetAllowTriggerFullPress(false);
-		// Consume light press
-		if ((Action.ActionType == ViewportWorldActionTypes::SelectAndMove_LightlyPressed) && (!VREditorMode->IsShowingRadialMenu(VRInteractor)))
+		const UVREditorInteractor* VREditorInteractor = Cast<UVREditorInteractor>(Interactor);
+		if (Action.ActionType == ViewportWorldActionTypes::SelectAndMove && ( VREditorInteractor == nullptr || !VREditorMode->IsShowingRadialMenu( VREditorInteractor )))
 		{
-			if ((Action.Event == IE_Pressed) && (!VRInteractor->IsHoveringOverPriorityType()))
+			if (Action.Event == IE_Pressed && !Interactor->IsHoveringOverPriorityType())
 			{
+				bWasHandled = true;
+				bOutIsInputCaptured = true;
+
 				// Go ahead and paint immediately
 				FVector LaserPointerStart, LaserPointerEnd;
 				if (Interactor->GetLaserPointer( /* Out */ LaserPointerStart, /* Out */ LaserPointerEnd))
@@ -565,6 +589,9 @@ void FEdModeFoliage::OnVRAction(class FEditorViewportClient& ViewportClient, UVi
 			{
 				EndFoliageBrushTrace();
 				FoliageInteractor = nullptr;
+				
+				bWasHandled = true;
+				bOutIsInputCaptured = false;
 			}
 		}
 	}
@@ -784,7 +811,7 @@ void FEdModeFoliage::EndFoliageBrushTrace()
 void FEdModeFoliage::FoliageBrushTrace(FEditorViewportClient* ViewportClient, const FVector& InRayOrigin, const FVector& InRayDirection)
 {
 	bBrushTraceValid = false;
-	if (!ViewportClient->IsMovingCamera() && ViewportClient->IsVisible())
+	if (ViewportClient == nullptr || ( !ViewportClient->IsMovingCamera() && ViewportClient->IsVisible() ) )
 	{
 		if (UISettings.GetPaintToolSelected() || UISettings.GetReapplyToolSelected() || UISettings.GetLassoSelectToolSelected())
 		{
@@ -792,7 +819,7 @@ void FEdModeFoliage::FoliageBrushTrace(FEditorViewportClient* ViewportClient, co
 			const FVector TraceEnd(InRayOrigin + InRayDirection * HALF_WORLD_MAX);
 
 			FHitResult Hit;
-			UWorld* World = ViewportClient->GetWorld();
+			UWorld* World = GetWorld();
 			static FName NAME_FoliageBrush = FName(TEXT("FoliageBrush"));
 			FFoliagePaintingGeometryFilter FilterFunc = FFoliagePaintingGeometryFilter(UISettings);
 
@@ -829,8 +856,8 @@ void FEdModeFoliage::FoliageBrushTrace(FEditorViewportClient* ViewportClient, co
 bool FEdModeFoliage::MouseMove(FEditorViewportClient* ViewportClient, FViewport* Viewport, int32 MouseX, int32 MouseY)
 {
 	// Use mouse capture if there's no other interactor currently tracing brush
-	UVREditorMode* VREditorMode = GEditor->GetEditorWorldManager()->GetEditorWorldWrapper(ViewportClient->GetWorld())->GetVREditorMode();
-	if (!VREditorMode->IsActive())
+	UVREditorMode* VREditorMode = Cast<UVREditorMode>( GEditor->GetEditorWorldExtensionsManager()->GetEditorWorldExtensions( GetWorld() )->FindExtension( UVREditorMode::StaticClass() ) );
+	if (VREditorMode == nullptr || !VREditorMode->IsActive())
 	{
 		// Compute a world space ray from the screen space mouse coordinates
 		FSceneViewFamilyContext ViewFamily(FSceneViewFamily::ConstructionValues(
@@ -867,8 +894,8 @@ bool FEdModeFoliage::MouseMove(FEditorViewportClient* ViewportClient, FViewport*
 bool FEdModeFoliage::CapturedMouseMove(FEditorViewportClient* ViewportClient, FViewport* Viewport, int32 MouseX, int32 MouseY)
 {
 	// Use mouse capture if there's no other interactor currently tracing brush
-	UVREditorMode* VREditorMode = GEditor->GetEditorWorldManager()->GetEditorWorldWrapper(ViewportClient->GetWorld())->GetVREditorMode();
-	if (!VREditorMode->IsActive())
+	UVREditorMode* VREditorMode = Cast<UVREditorMode>( GEditor->GetEditorWorldExtensionsManager()->GetEditorWorldExtensions( GetWorld() )->FindExtension( UVREditorMode::StaticClass() ) );
+	if (VREditorMode == nullptr || !VREditorMode->IsActive())
 	{
 		//Compute a world space ray from the screen space mouse coordinates
 		FSceneViewFamilyContext ViewFamily(FSceneViewFamily::ConstructionValues(
@@ -1118,6 +1145,12 @@ bool FEdModeFoliage::VertexMaskCheck(const FHitResult& Hit, const UFoliageType*
 	return true;
 }
 
+void FEdModeFoliage::SetBrushOpacity(const float InOpacity)
+{
+	static FName OpacityParamName("OpacityAmount");
+	BrushMID->SetScalarParameterValue(OpacityParamName, InOpacity);
+}
+
 bool LandscapeLayerCheck(const FHitResult& Hit, const UFoliageType* Settings, FEdModeFoliage::LandscapeLayerCacheData& LandscapeLayersCache, float& OutHitWeight)
 {
 	OutHitWeight = 1.f;
@@ -3461,4 +3494,4 @@ void FFoliageUISettings::Save()
 	GConfig->SetFloat(TEXT("FoliageEdit"), TEXT("PaletteThumbnailScale"), PaletteThumbnailScale, GEditorPerProjectIni);
 }
 
-#undef LOCTEXT_NAMESPACE
+#undef LOCTEXT_NAMESPACE
\ No newline at end of file
diff --git a/Engine/Source/Editor/FoliageEdit/Private/FoliageEdMode.h b/Engine/Source/Editor/FoliageEdit/Private/FoliageEdMode.h
--- a/Engine/Source/Editor/FoliageEdit/Private/FoliageEdMode.h
+++ b/Engine/Source/Editor/FoliageEdit/Private/FoliageEdMode.h
@@ -116,7 +116,7 @@ struct FFoliageUISettings
 		, bShowPaletteItemTooltips(true)
 		, ActivePaletteViewMode(EFoliagePaletteViewMode::Thumbnail)
 		, PaletteThumbnailScale(0.3f)
-		, Radius(512.f)
+		, Radius(250.f)
 		, PaintDensity(0.5f)
 		, UnpaintDensity(0.f)
 		, IsInSingleInstantiationMode(false)
@@ -456,7 +456,7 @@ public:
 	void OnVRAction(class FEditorViewportClient& ViewportClient, class UViewportInteractor* Interactor, const struct FViewportActionKeyInput& Action, bool& bOutIsInputCaptured, bool& bWasHandled);
 
 	/** Called on VR hovering */
-	void OnVRHoverUpdate(FEditorViewportClient& ViewportClient, UViewportInteractor* Interactor, FVector& HoverImpactPoint, bool& bWasHandled);
+	void OnVRHoverUpdate(UViewportInteractor* Interactor, FVector& HoverImpactPoint, bool& bWasHandled);
 
 	typedef TMap<FName, TMap<ULandscapeComponent*, TArray<uint8> > > LandscapeLayerCacheData;
 
@@ -559,12 +559,21 @@ private:
 	/** Does a filter based on the vertex color of a static mesh */
 	static bool VertexMaskCheck(const FHitResult& Hit, const UFoliageType* Settings);
 
+	/** Set the brush mesh opacity */
+	void SetBrushOpacity(const float InOpacity);
+
 	bool bBrushTraceValid;
 	FVector BrushLocation;
 	FVector BrushNormal;
 	FVector BrushTraceDirection;
 	UStaticMeshComponent* SphereBrushComponent;
 
+	/** The dynamic material of the sphere brush. */
+	class UMaterialInstanceDynamic* BrushMID;
+
+	/** Default opacity received from the brush material to reset it when closing. */
+	float DefaultBrushOpacity;
+
 	// Landscape layer cache data
 	LandscapeLayerCacheData LandscapeLayerCaches;
 
diff --git a/Engine/Source/Editor/GeometryMode/Private/GeometryEdMode.cpp b/Engine/Source/Editor/GeometryMode/Private/GeometryEdMode.cpp
--- a/Engine/Source/Editor/GeometryMode/Private/GeometryEdMode.cpp
+++ b/Engine/Source/Editor/GeometryMode/Private/GeometryEdMode.cpp
@@ -1098,7 +1098,7 @@ bool FModeTool_GeometryModify::BoxSelect( FBox& InBox, bool InSelect )
 }
 
 /** @return		true if something was selected/deselected, false otherwise. */
-bool FModeTool_GeometryModify::FrustumSelect( const FConvexVolume& InFrustum, bool InSelect /* = true */ )
+bool FModeTool_GeometryModify::FrustumSelect( const FConvexVolume& InFrustum, FEditorViewportClient* InViewportClient, bool InSelect /* = true */ )
 {
 	bool bResult = false;
 	if( GLevelEditorModeTools().IsModeActive( FBuiltinEditorModes::EM_Geometry ) )
diff --git a/Engine/Source/Editor/GeometryMode/Public/GeometryEdMode.h b/Engine/Source/Editor/GeometryMode/Public/GeometryEdMode.h
--- a/Engine/Source/Editor/GeometryMode/Public/GeometryEdMode.h
+++ b/Engine/Source/Editor/GeometryMode/Public/GeometryEdMode.h
@@ -250,7 +250,7 @@ public:
 
 	virtual void SelectNone() override;
 	virtual bool BoxSelect( FBox& InBox, bool InSelect = true ) override;
-	virtual bool FrustumSelect( const FConvexVolume& InFrustum, bool InSelect = true ) override;
+	virtual bool FrustumSelect( const FConvexVolume& InFrustum, FEditorViewportClient* InViewportClient, bool InSelect = true ) override;
 
 	/** @name Modifier iterators */
 	//@{
diff --git a/Engine/Source/Editor/LevelEditor/Private/LevelEditorToolBar.cpp b/Engine/Source/Editor/LevelEditor/Private/LevelEditorToolBar.cpp
--- a/Engine/Source/Editor/LevelEditor/Private/LevelEditorToolBar.cpp
+++ b/Engine/Source/Editor/LevelEditor/Private/LevelEditorToolBar.cpp
@@ -1269,7 +1269,7 @@ TSharedRef< SWidget > FLevelEditorToolBar::MakeLevelEditorToolBar( const TShared
 			FSlateIcon(FEditorStyle::GetStyleSetName(), "LevelEditor.EditMatinee") 
 			);
 
-		ToolbarBuilder.AddToolBarButton( FLevelEditorCommands::Get().ToggleVR, NAME_None, LOCTEXT("ToggleVR", "VR") );
+		ToolbarBuilder.AddToolBarButton( FLevelEditorCommands::Get().ToggleVR, NAME_None, LOCTEXT("ToggleVR", "VR Editor") );
 	}
 	ToolbarBuilder.EndSection();
 	
diff --git a/Engine/Source/Editor/Matinee/Private/MatineeTransaction.cpp b/Engine/Source/Editor/Matinee/Private/MatineeTransaction.cpp
--- a/Engine/Source/Editor/Matinee/Private/MatineeTransaction.cpp
+++ b/Engine/Source/Editor/Matinee/Private/MatineeTransaction.cpp
@@ -10,9 +10,7 @@
 #include "Matinee/InterpGroupInst.h"
 #include "Matinee/InterpTrackInst.h"
 
-/*-----------------------------------------------------------------------------
-	FMatineeTransaction
------------------------------------------------------------------------------*/
+
 void FMatineeTransaction::SaveObject( UObject* Object )
 {
 	check(Object);
@@ -27,7 +25,7 @@ void FMatineeTransaction::SaveObject( UObject* Object )
 		Object->IsA( UK2Node_MatineeController::StaticClass() ) )
 	{
 		// Save the object.
-		new( Records )FObjectRecord( this, Object, NULL, 0, 0, 0, 0, NULL, NULL, NULL );
+		new( Records )FObjectRecord( this, Object, nullptr, NULL, 0, 0, 0, 0, NULL, NULL, NULL );
 	}
 }
 
@@ -35,3 +33,8 @@ void FMatineeTransaction::SaveArray( UObject* Object, FScriptArray* Array, int32
 {
 	// Never want this.
 }
+
+void FMatineeTransaction::StoreUndo( UObject* Object, TUniquePtr<FChange> CustomChange )
+{
+	// Not used
+}
diff --git a/Engine/Source/Editor/Matinee/Private/MatineeTransaction.h b/Engine/Source/Editor/Matinee/Private/MatineeTransaction.h
--- a/Engine/Source/Editor/Matinee/Private/MatineeTransaction.h
+++ b/Engine/Source/Editor/Matinee/Private/MatineeTransaction.h
@@ -22,5 +22,6 @@ public:
 
 	virtual void SaveObject( UObject* Object ) override;
 	virtual void SaveArray( UObject* Object, FScriptArray* Array, int32 Index, int32 Count, int32 Oper, int32 ElementSize, STRUCT_DC DefaultConstructor, STRUCT_AR Serializer, STRUCT_DTOR Destructor ) override;
+	virtual void StoreUndo( UObject* Object, TUniquePtr<FChange> UndoChange ) override;
 };
 
diff --git a/Engine/Source/Editor/MeshEditor/MeshEditor.Build.cs b/Engine/Source/Editor/MeshEditor/MeshEditor.Build.cs
new file mode 100644
--- /dev/null
+++ b/Engine/Source/Editor/MeshEditor/MeshEditor.Build.cs
@@ -0,0 +1,48 @@
+// Copyright 1998-2015 Epic Games, Inc. All Rights Reserved.
+
+namespace UnrealBuildTool.Rules
+{
+	public class MeshEditor : ModuleRules
+	{
+        public MeshEditor(TargetInfo Target)
+		{
+			PrivateDependencyModuleNames.AddRange(
+                new string[] {
+				    "Core",
+				    "CoreUObject",
+				    "Engine",
+                    "InputCore",
+				    "Slate",
+					"SlateCore",
+                    "EditorStyle",
+				    "UnrealEd",
+					"RenderCore",
+					"MeshEditingRuntime",
+					"ViewportInteraction",
+                    "VREditor"
+				}
+			);
+
+			CircularlyReferencedDependentModules.AddRange(
+				new string[]
+				{
+					"ViewportInteraction"
+				}
+			);
+
+			PrivateIncludePathModuleNames.AddRange(
+				new string[] {
+                    "ContentBrowser",
+					"LevelEditor",
+				}
+			);
+
+			DynamicallyLoadedModuleNames.AddRange(
+				new string[] {
+                    "ContentBrowser",
+					"LevelEditor",
+				}
+			);
+		}
+	}
+}
\ No newline at end of file
diff --git a/Engine/Source/Editor/MeshEditor/MeshEditorCommands.cpp b/Engine/Source/Editor/MeshEditor/MeshEditorCommands.cpp
new file mode 100644
--- /dev/null
+++ b/Engine/Source/Editor/MeshEditor/MeshEditorCommands.cpp
@@ -0,0 +1,64 @@
+// Copyright 1998-2015 Epic Games, Inc. All Rights Reserved.
+
+#include "MeshEditorCommands.h"
+
+#define LOCTEXT_NAMESPACE "MeshEditorCommands"
+
+void FMeshEditorCommonCommands::RegisterCommands()
+{
+	UI_COMMAND(DeleteMeshElement, "Delete", "Delete selected mesh elements, including polygons partly defined by selected elements.", EUserInterfaceActionType::Button, FInputChord(EKeys::Delete));
+	UI_COMMAND(AddSubdivisionLevel, "Add Subdivision Level", "Increases the number of subdivision levels for the selected mesh.", EUserInterfaceActionType::Button, FInputChord(EKeys::Equals));
+	UI_COMMAND(RemoveSubdivisionLevel, "Remove Subdivision Level", "Decreases the number of subdivision levels for the selected mesh.", EUserInterfaceActionType::Button, FInputChord(EKeys::Hyphen));
+	UI_COMMAND(ShowVertexNormals, "Show Vertex Normals", "Toggles debug rendering of vertex normals.", EUserInterfaceActionType::ToggleButton, FInputChord(EKeys::V));
+	UI_COMMAND(MarqueeSelectVertices, "Marquee Select Vertices", "Selects vertices inside the selection box.", EUserInterfaceActionType::None, FInputChord());
+	UI_COMMAND(MarqueeSelectEdges, "Marquee Select Edges", "Selects edges inside the selection box.", EUserInterfaceActionType::None, FInputChord());
+	UI_COMMAND(MarqueeSelectPolygons, "Marquee Select Polygons", "Selects polygons inside the selection box.", EUserInterfaceActionType::None, FInputChord());
+	UI_COMMAND(DrawVertices, "Draw Vertices", "Allows vertices to be freely drawn to create a new polygon.", EUserInterfaceActionType::RadioButton, FInputChord());
+	UI_COMMAND(FrameSelectedElements, "Frame Selected Elements", "Moves the viewport camera to frame the currently selected elements.", EUserInterfaceActionType::None, FInputChord(EKeys::F));
+	UI_COMMAND(SetVertexSelectionMode, "Set Vertex Selection Mode", "Sets the selection mode so that only vertices will be selected.", EUserInterfaceActionType::None, FInputChord(EKeys::One));
+	UI_COMMAND(SetEdgeSelectionMode, "Set Edge Selection Mode", "Sets the selection mode so that only edges will be selected.", EUserInterfaceActionType::None, FInputChord(EKeys::Two));
+	UI_COMMAND(SetPolygonSelectionMode, "Set Polygon Selection Mode", "Sets the selection mode so that only polygons will be selected.", EUserInterfaceActionType::None, FInputChord(EKeys::Three));
+	UI_COMMAND(SetAnySelectionMode, "Set Any Selection Mode", "Sets the selection mode so that any element type may be selected.", EUserInterfaceActionType::None, FInputChord(EKeys::Four));
+	UI_COMMAND(QuadrangulateMesh, "Quadrangulate Mesh", "Quadrangulates the selected mesh.", EUserInterfaceActionType::Button, FInputChord());
+}
+
+void FMeshEditorVertexCommands::RegisterCommands()
+{
+	UI_COMMAND(MoveVertex, "Move Vertex Mode", "Set the primary action to move vertices.", EUserInterfaceActionType::RadioButton, FInputChord(EKeys::F1));
+	UI_COMMAND(ExtendVertex, "Extend Vertex Mode", "Set the primary action to extend vertices.", EUserInterfaceActionType::RadioButton, FInputChord(EKeys::F2));
+	UI_COMMAND(EditVertexCornerSharpness, "Edit Vertex Corner Sharpness Mode", "Set the primary action to edit the vertex's subdivision corner sharpness amount.", EUserInterfaceActionType::RadioButton, FInputChord(EKeys::F6));
+
+	UI_COMMAND(RemoveVertex, "Remove Vertex", "Remove the selected vertex if possible.", EUserInterfaceActionType::Button, FInputChord(EKeys::BackSpace));
+	UI_COMMAND(WeldVertices, "Weld Vertices", "Weld the selected vertices, keeping the first selected vertex.", EUserInterfaceActionType::Button, FInputChord());
+}
+
+void FMeshEditorEdgeCommands::RegisterCommands()
+{
+	UI_COMMAND(MoveEdge, "Move Edge Mode", "Set the primary action to move edges.", EUserInterfaceActionType::RadioButton, FInputChord(EKeys::F1));
+	UI_COMMAND(SplitEdge, "Split Edge Mode", "Set the primary action to split edges.", EUserInterfaceActionType::RadioButton, FInputChord(EKeys::F2));
+	UI_COMMAND(SplitEdgeAndDragVertex, "Split Edge and Drag Vertex Mode", "Set the primary action to split edges and drag vertices.", EUserInterfaceActionType::RadioButton, FInputChord(EKeys::F3));
+	UI_COMMAND(InsertEdgeLoop, "Insert Edge Loop Mode", "Set the primary action to insert edge loops.", EUserInterfaceActionType::RadioButton, FInputChord(EKeys::F4));
+	UI_COMMAND(ExtendEdge, "Extend Edge Mode", "Set the primary action to extend edges.", EUserInterfaceActionType::RadioButton, FInputChord(EKeys::F5));
+	UI_COMMAND(EditEdgeCreaseSharpness, "Edit Edge Crease Sharpness Mode", "Set the primary action to edit the edge's crease sharpness.", EUserInterfaceActionType::RadioButton, FInputChord(EKeys::F6));
+
+	UI_COMMAND(RemoveEdge, "Remove Edge", "Remove the selected edge if possible.", EUserInterfaceActionType::Button, FInputChord(EKeys::BackSpace));
+	UI_COMMAND(SoftenEdge, "Soften Edge", "Make selected edge soft.", EUserInterfaceActionType::Button, FInputChord(EKeys::H, EModifierKey::Shift));
+	UI_COMMAND(HardenEdge, "Harden Edge", "Make selected edge hard.", EUserInterfaceActionType::Button, FInputChord(EKeys::H));
+	UI_COMMAND(SelectEdgeLoop, "Select Edge Loop", "Select the edge loops which contain the selected edges.", EUserInterfaceActionType::Button, FInputChord(EKeys::Two, EModifierKey::Shift));
+}
+
+void FMeshEditorPolygonCommands::RegisterCommands()
+{
+	UI_COMMAND(MovePolygon, "Move Polygon Mode", "Set the primary action to move polygons.", EUserInterfaceActionType::RadioButton, FInputChord(EKeys::F1));
+	UI_COMMAND(ExtrudePolygon, "Extrude Polygon Mode", "Set the primary action to extrude polygons.", EUserInterfaceActionType::RadioButton, FInputChord(EKeys::F2));
+	// UI_COMMAND(FreelyExtrudePolygon, "Freely Extrude Polygon Mode", "Set the primary action to freely extrude polygons.", EUserInterfaceActionType::RadioButton, FInputChord(EKeys::F3));
+	UI_COMMAND(InsetPolygon, "Inset Polygon Mode", "Set the primary action to inset polygons.", EUserInterfaceActionType::RadioButton, FInputChord(EKeys::F3));
+	UI_COMMAND(BevelPolygon, "Bevel Polygon Mode", "Set the primary action to bevel polygons.", EUserInterfaceActionType::RadioButton, FInputChord(EKeys::F4));
+
+	UI_COMMAND(FlipPolygon, "Flip Polygon", "Flip the currently selected polygons.", EUserInterfaceActionType::Button, FInputChord(EKeys::F, EModifierKey::Shift));
+	UI_COMMAND(TriangulatePolygon, "Triangulate Polygon", "Triangulate the currently selected polygons.", EUserInterfaceActionType::Button, FInputChord(EKeys::T));
+	UI_COMMAND(TessellatePolygon, "Tessellate Selected Polygons", "Tessellate selected polygons into smaller polygons.", EUserInterfaceActionType::Button, FInputChord());
+	UI_COMMAND(AssignMaterial, "Assign Material", "Assigns the highlighted material in the Content Browser to the currently selected polygons.", EUserInterfaceActionType::Button, FInputChord(EKeys::M));
+}
+
+#undef LOCTEXT_NAMESPACE
\ No newline at end of file
