[NO NON-PII COMMIT MESSAGE]

diff --git a/Engine/Source/Editor/FoliageEdit/Private/FoliageEdMode.cpp b/Engine/Source/Editor/FoliageEdit/Private/FoliageEdMode.cpp
--- a/Engine/Source/Editor/FoliageEdit/Private/FoliageEdMode.cpp
+++ b/Engine/Source/Editor/FoliageEdit/Private/FoliageEdMode.cpp
@@ -1398,7 +1398,11 @@ void FEdModeFoliage::CalculatePotentialInstances(const UWorld* InWorld, const UF
 			float HitWeight = 1.f;
 
 			UPrimitiveComponent* InstanceBase = Hit.GetComponent();
-			check(InstanceBase);
+
+			if (InstanceBase == nullptr)
+			{
+				continue;
+			}
 
 			ULevel* TargetLevel = InstanceBase->GetComponentLevel();
 			// We can paint into new level only if FoliageType is shared
