Fixed Editor instance count, to only include placed instances.
Calling GenerateProceduralContent will no longer automatically remove the current content, as we might want to get the generation for something else than replacing existing data.
When deleting procedural foliage, we can now specify if we want to rebuild the tree.

diff --git a/Engine/Source/Editor/FoliageEdit/Private/FoliageEdMode.cpp b/Engine/Source/Editor/FoliageEdit/Private/FoliageEdMode.cpp
--- a/Engine/Source/Editor/FoliageEdit/Private/FoliageEdMode.cpp
+++ b/Engine/Source/Editor/FoliageEdit/Private/FoliageEdMode.cpp
@@ -2981,10 +2981,12 @@ void FEdModeFoliage::PopulateFoliageMeshList()
 						ElementIdx = FoliageMeshList.Add(MakeShareable(new FFoliageMeshUIInfo(MeshPair.Key)));
 					}
 
-					FoliageMeshList[ElementIdx]->InstanceCountTotal += MeshPair.Value->GetInstanceCount();
+					int32 PlacedInstanceCount = MeshPair.Value->GetPlacedInstanceCount();
+					FoliageMeshList[ElementIdx]->InstanceCountTotal += PlacedInstanceCount;
+
 					if (Level == World->GetCurrentLevel())
 					{
-						FoliageMeshList[ElementIdx]->InstanceCountCurrentLevel += MeshPair.Value->GetInstanceCount();
+						FoliageMeshList[ElementIdx]->InstanceCountCurrentLevel += PlacedInstanceCount;
 					}
 				}
 			}
diff --git a/Engine/Source/Editor/FoliageEdit/Private/ProceduralFoliageComponentDetails.cpp b/Engine/Source/Editor/FoliageEdit/Private/ProceduralFoliageComponentDetails.cpp
--- a/Engine/Source/Editor/FoliageEdit/Private/ProceduralFoliageComponentDetails.cpp
+++ b/Engine/Source/Editor/FoliageEdit/Private/ProceduralFoliageComponentDetails.cpp
@@ -80,14 +80,19 @@ FReply FProceduralFoliageComponentDetails::OnResimulateClicked()
 			TArray <FDesiredFoliageInstance> DesiredFoliageInstances;
 			if (Component->GenerateProceduralContent(DesiredFoliageInstances))
 			{
-				FFoliagePaintingGeometryFilter OverrideGeometryFilter;
-				OverrideGeometryFilter.bAllowLandscape = Component->bAllowLandscape;
-				OverrideGeometryFilter.bAllowStaticMesh = Component->bAllowStaticMesh;
-				OverrideGeometryFilter.bAllowBSP = Component->bAllowBSP;
-				OverrideGeometryFilter.bAllowFoliage = Component->bAllowFoliage;
-				OverrideGeometryFilter.bAllowTranslucent = Component->bAllowTranslucent;
-
-				FEdModeFoliage::AddInstances(Component->GetWorld(), DesiredFoliageInstances, OverrideGeometryFilter);
+				if (DesiredFoliageInstances.Num() > 0)
+				{
+					Component->RemoveProceduralContent();
+
+					FFoliagePaintingGeometryFilter OverrideGeometryFilter;
+					OverrideGeometryFilter.bAllowLandscape = Component->bAllowLandscape;
+					OverrideGeometryFilter.bAllowStaticMesh = Component->bAllowStaticMesh;
+					OverrideGeometryFilter.bAllowBSP = Component->bAllowBSP;
+					OverrideGeometryFilter.bAllowFoliage = Component->bAllowFoliage;
+					OverrideGeometryFilter.bAllowTranslucent = Component->bAllowTranslucent;
+
+					FEdModeFoliage::AddInstances(Component->GetWorld(), DesiredFoliageInstances, OverrideGeometryFilter);
+				}
 
 				// If no instances were spawned, inform the user
 				if (!Component->HasSpawnedAnyInstances())
diff --git a/Engine/Source/Runtime/Foliage/Private/InstancedFoliage.cpp b/Engine/Source/Runtime/Foliage/Private/InstancedFoliage.cpp
--- a/Engine/Source/Runtime/Foliage/Private/InstancedFoliage.cpp
+++ b/Engine/Source/Runtime/Foliage/Private/InstancedFoliage.cpp
@@ -1275,9 +1275,19 @@ void FFoliageMeshInfo::DuplicateInstances(AInstancedFoliageActor* InIFA, UFoliag
 }
 
 /* Get the number of placed instances */
-int32 FFoliageMeshInfo::GetInstanceCount() const
+int32 FFoliageMeshInfo::GetPlacedInstanceCount() const
 {
-	return Instances.Num();
+	int32 PlacedInstanceCount = 0;
+
+	for (int32 i = 0; i < Instances.Num(); ++i)
+	{
+		if (!Instances[i].ProceduralGuid.IsValid())
+		{
+			++PlacedInstanceCount;
+		}
+	}
+
+	return PlacedInstanceCount;
 }
 
 void FFoliageMeshInfo::AddToBaseHash(int32 InstanceIndex)
@@ -1800,7 +1810,7 @@ void AInstancedFoliageActor::DeleteInstancesForComponent(UWorld* InWorld, UActor
 	}
 }
 
-void AInstancedFoliageActor::DeleteInstancesForProceduralFoliageComponent(const UProceduralFoliageComponent* ProceduralFoliageComponent)
+void AInstancedFoliageActor::DeleteInstancesForProceduralFoliageComponent(const UProceduralFoliageComponent* ProceduralFoliageComponent, bool InRebuildTree)
 {
 	const FGuid& ProceduralGuid = ProceduralFoliageComponent->GetProceduralGuid();
 	for (auto& MeshPair : FoliageMeshes)
@@ -1817,7 +1827,7 @@ void AInstancedFoliageActor::DeleteInstancesForProceduralFoliageComponent(const
 
 		if (InstancesToRemove.Num())
 		{
-			MeshInfo.RemoveInstances(this, InstancesToRemove, true);
+			MeshInfo.RemoveInstances(this, InstancesToRemove, InRebuildTree);
 		}
 	}
 }
diff --git a/Engine/Source/Runtime/Foliage/Private/ProceduralFoliageComponent.cpp b/Engine/Source/Runtime/Foliage/Private/ProceduralFoliageComponent.cpp
--- a/Engine/Source/Runtime/Foliage/Private/ProceduralFoliageComponent.cpp
+++ b/Engine/Source/Runtime/Foliage/Private/ProceduralFoliageComponent.cpp
@@ -258,19 +258,13 @@ void UProceduralFoliageComponent::PostEditImport()
 bool UProceduralFoliageComponent::GenerateProceduralContent(TArray <FDesiredFoliageInstance>& OutInstances)
 {
 #if WITH_EDITOR
-	
-	if (ExecuteSimulation(OutInstances))
-	{
-		RemoveProceduralContent();
-		return true;
-	}
-
+	return ExecuteSimulation(OutInstances);
 #endif
 
 	return false;
 }
 
-void UProceduralFoliageComponent::RemoveProceduralContent()
+void UProceduralFoliageComponent::RemoveProceduralContent(bool InRebuildTree)
 {
 #if WITH_EDITOR
 	UWorld* World = GetWorld();
@@ -282,7 +276,7 @@ void UProceduralFoliageComponent::RemoveProceduralContent()
 			AInstancedFoliageActor* IFA = AInstancedFoliageActor::GetInstancedFoliageActorForLevel(Level);
 			if (IFA)
 			{
-				IFA->DeleteInstancesForProceduralFoliageComponent(this);
+				IFA->DeleteInstancesForProceduralFoliageComponent(this, InRebuildTree);
 			}
 		}
 	}
diff --git a/Engine/Source/Runtime/Foliage/Private/ProceduralFoliageTile.cpp b/Engine/Source/Runtime/Foliage/Private/ProceduralFoliageTile.cpp
--- a/Engine/Source/Runtime/Foliage/Private/ProceduralFoliageTile.cpp
+++ b/Engine/Source/Runtime/Foliage/Private/ProceduralFoliageTile.cpp
@@ -406,10 +406,7 @@ void UProceduralFoliageTile::StepSimulation()
 		SpreadSeeds(NewInstances);
 	}
 
-	for (FProceduralFoliageInstance* Inst : NewInstances)
-	{
-		InstancesSet.Add(Inst);
-	}
+	InstancesSet.Append(NewInstances);
 
 	FlushPendingRemovals();
 }
diff --git a/Engine/Source/Runtime/Foliage/Public/InstancedFoliage.h b/Engine/Source/Runtime/Foliage/Public/InstancedFoliage.h
--- a/Engine/Source/Runtime/Foliage/Public/InstancedFoliage.h
+++ b/Engine/Source/Runtime/Foliage/Public/InstancedFoliage.h
@@ -83,6 +83,7 @@ struct FFoliageInstance : public FFoliageInstancePlacementInfo
 
 	FFoliageInstance()
 		: BaseId(0)
+		, BaseComponent(nullptr)
 	{}
 
 	friend FArchive& operator<<(FArchive& Ar, FFoliageInstance& Instance);
@@ -224,13 +225,13 @@ struct FFoliageMeshInfo
 	FOLIAGE_API void SelectInstances(AInstancedFoliageActor* InIFA, bool bSelect);
 
 	// Get the number of placed instances
-	FOLIAGE_API int32 GetInstanceCount() const;
+	FOLIAGE_API int32 GetPlacedInstanceCount() const;
 
 	FOLIAGE_API void AddToBaseHash(int32 InstanceIdx);
 	FOLIAGE_API void RemoveFromBaseHash(int32 InstanceIdx);
 
 	// Create and register a new component
-	void CreateNewComponent(AInstancedFoliageActor* InIFA, const UFoliageType* InSettings);
+	FOLIAGE_API void CreateNewComponent(AInstancedFoliageActor* InIFA, const UFoliageType* InSettings);
 
 	// For debugging. Validate state after editing.
 	void CheckValid();
diff --git a/Engine/Source/Runtime/Foliage/Public/InstancedFoliageActor.h b/Engine/Source/Runtime/Foliage/Public/InstancedFoliageActor.h
--- a/Engine/Source/Runtime/Foliage/Public/InstancedFoliageActor.h
+++ b/Engine/Source/Runtime/Foliage/Public/InstancedFoliageActor.h
@@ -142,7 +142,7 @@ public:
 	static FOLIAGE_API void DeleteInstancesForComponent(UWorld* InWorld, UActorComponent* InComponent);
 
 	// Deletes the instances spawned by a procedural component
-	void DeleteInstancesForProceduralFoliageComponent(const UProceduralFoliageComponent* ProceduralFoliageComponent);
+	void DeleteInstancesForProceduralFoliageComponent(const UProceduralFoliageComponent* ProceduralFoliageComponent, bool InRebuildTree);
 
 	/** @return True if any instances exist that were spawned by the given procedural component */
 	bool ContainsInstancesFromProceduralFoliageComponent(const UProceduralFoliageComponent* ProceduralFoliageComponent);
diff --git a/Engine/Source/Runtime/Foliage/Public/ProceduralFoliageComponent.h b/Engine/Source/Runtime/Foliage/Public/ProceduralFoliageComponent.h
--- a/Engine/Source/Runtime/Foliage/Public/ProceduralFoliageComponent.h
+++ b/Engine/Source/Runtime/Foliage/Public/ProceduralFoliageComponent.h
@@ -87,7 +87,7 @@ class FOLIAGE_API UProceduralFoliageComponent : public UActorComponent
 	bool GenerateProceduralContent(TArray<FDesiredFoliageInstance>& OutFoliageInstances);
 	
 	/** Removes all spawned foliage instances in the level that were spawned by this component */
-	void RemoveProceduralContent();
+	void RemoveProceduralContent(bool InRebuildTree = true);
 
 	/** @return True if any foliage instances in the level were spawned by this component */
 	bool HasSpawnedAnyInstances();
