Fixed a scale issue that would give a much lower fistance factor than expected

diff --git a/Engine/Source/Editor/FoliageEdit/Private/FoliageEdMode.cpp b/Engine/Source/Editor/FoliageEdit/Private/FoliageEdMode.cpp
--- a/Engine/Source/Editor/FoliageEdit/Private/FoliageEdMode.cpp
+++ b/Engine/Source/Editor/FoliageEdit/Private/FoliageEdMode.cpp
@@ -1423,7 +1423,7 @@ void FEdModeFoliage::CalculatePotentialInstances(const UWorld* InWorld, const UF
 	}
 }
 
-void FEdModeFoliage::AddInstances(UWorld* InWorld, const TArray<FDesiredFoliageInstance>& DesiredInstances, const FFoliagePaintingGeometryFilter& OverrideGeometryFilter)
+void FEdModeFoliage::AddInstances(UWorld* InWorld, const TArray<FDesiredFoliageInstance>& DesiredInstances, const FFoliagePaintingGeometryFilter& OverrideGeometryFilter, bool InRebuildFoliageTree)
 {
 	TMap<const UFoliageType*, TArray<FDesiredFoliageInstance>> SettingsInstancesMap;
 	for (const FDesiredFoliageInstance& DesiredInst : DesiredInstances)
@@ -1437,11 +1437,11 @@ void FEdModeFoliage::AddInstances(UWorld* InWorld, const TArray<FDesiredFoliageI
 		const UFoliageType* FoliageType = It.Key();
 
 		const TArray<FDesiredFoliageInstance>& Instances = It.Value();
-		AddInstancesImp(InWorld, FoliageType, Instances, TArray<int32>(), 1.f, nullptr, nullptr, &OverrideGeometryFilter);
+		AddInstancesImp(InWorld, FoliageType, Instances, TArray<int32>(), 1.f, nullptr, nullptr, &OverrideGeometryFilter, InRebuildFoliageTree);
 	}
 }
 
-static void SpawnFoliageInstance(UWorld* InWorld, const UFoliageType* Settings, const FFoliageUISettings* UISettings, const TArray<FFoliageInstance>& PlacedInstances)
+static void SpawnFoliageInstance(UWorld* InWorld, const UFoliageType* Settings, const FFoliageUISettings* UISettings, const TArray<FFoliageInstance>& PlacedInstances, bool InRebuildFoliageTree)
 {
 	SCOPE_CYCLE_COUNTER(STAT_FoliageSpawnInstance);
 
@@ -1476,6 +1476,11 @@ static void SpawnFoliageInstance(UWorld* InWorld, const UFoliageType* Settings,
 		UFoliageType* FoliageSettings = IFA->AddFoliageType(Settings, &MeshInfo);
 
 		MeshInfo->AddInstances(IFA, FoliageSettings, PlacedLevelInstances.Value, false);
+
+		if (MeshInfo->Component != nullptr && InRebuildFoliageTree)
+		{
+			MeshInfo->Component->BuildTreeIfOutdated(true, false);
+		}
 	}
 }
 
@@ -1497,7 +1502,7 @@ void FEdModeFoliage::RebuildFoliageTree(const UFoliageType* Settings)
 	}
 }
 
-void FEdModeFoliage::AddInstancesImp(UWorld* InWorld, const UFoliageType* Settings, const TArray<FDesiredFoliageInstance>& DesiredInstances, const TArray<int32>& ExistingInstanceBuckets, const float Pressure, LandscapeLayerCacheData* LandscapeLayerCachesPtr, const FFoliageUISettings* UISettings, const FFoliagePaintingGeometryFilter* OverrideGeometryFilter)
+void FEdModeFoliage::AddInstancesImp(UWorld* InWorld, const UFoliageType* Settings, const TArray<FDesiredFoliageInstance>& DesiredInstances, const TArray<int32>& ExistingInstanceBuckets, const float Pressure, LandscapeLayerCacheData* LandscapeLayerCachesPtr, const FFoliageUISettings* UISettings, const FFoliagePaintingGeometryFilter* OverrideGeometryFilter, bool InRebuildFoliageTree)
 {
 	SCOPE_CYCLE_COUNTER(STAT_FoliageAddInstanceImp);
 	
@@ -1564,7 +1569,7 @@ void FEdModeFoliage::AddInstancesImp(UWorld* InWorld, const UFoliageType* Settin
 				}
 			}
 
-			SpawnFoliageInstance(InWorld, Settings, UISettings, PlacedInstances);
+			SpawnFoliageInstance(InWorld, Settings, UISettings, PlacedInstances, InRebuildFoliageTree);
 		}
 	}
 }
@@ -1586,7 +1591,7 @@ void FEdModeFoliage::AddSingleInstanceForBrush(UWorld* InWorld, const UFoliageTy
 	TArray<int32> ExistingInstanceBuckets;
 	ExistingInstanceBuckets.AddZeroed(NUM_INSTANCE_BUCKETS);
 
-	AddInstancesImp(InWorld, Settings, DesiredInstances, ExistingInstanceBuckets, Pressure, &LandscapeLayerCaches, &UISettings);
+	AddInstancesImp(InWorld, Settings, DesiredInstances, ExistingInstanceBuckets, Pressure, &LandscapeLayerCaches, &UISettings, false);
 }
 
 /** Add instances inside the brush to match DesiredInstanceCount */
@@ -1642,7 +1647,7 @@ void FEdModeFoliage::AddInstancesForBrush(UWorld* InWorld, const UFoliageType* S
 			FDesiredFoliageInstance* DesiredInstance = new (DesiredInstances)FDesiredFoliageInstance(Start, End);
 		}
 
-		AddInstancesImp(InWorld, Settings, DesiredInstances, ExistingInstanceBuckets, Pressure, &LandscapeLayerCaches, &UISettings);
+		AddInstancesImp(InWorld, Settings, DesiredInstances, ExistingInstanceBuckets, Pressure, &LandscapeLayerCaches, &UISettings, false);
 	}
 }
 
@@ -2717,7 +2722,7 @@ void FEdModeFoliage::ApplyPaintBucket_Add(AActor* Actor)
 				}
 			}
 
-			SpawnFoliageInstance(World, Settings, &UISettings, PlacedInstances);
+			SpawnFoliageInstance(World, Settings, &UISettings, PlacedInstances, false);
 		}
 
 		RebuildFoliageTree(Settings);
diff --git a/Engine/Source/Editor/FoliageEdit/Private/FoliageEdMode.h b/Engine/Source/Editor/FoliageEdit/Private/FoliageEdMode.h
--- a/Engine/Source/Editor/FoliageEdit/Private/FoliageEdMode.h
+++ b/Engine/Source/Editor/FoliageEdit/Private/FoliageEdMode.h
@@ -463,7 +463,7 @@ public:
 	void AdjustBrushRadius(float Adjustment);
 
 	/** Add desired instances. Uses foliage settings to determine location/scale/rotation and whether instances should be ignored */
-	static void AddInstances(UWorld* InWorld, const TArray<FDesiredFoliageInstance>& DesiredInstances, const FFoliagePaintingGeometryFilter& OverrideGeometryFilter);
+	static void AddInstances(UWorld* InWorld, const TArray<FDesiredFoliageInstance>& DesiredInstances, const FFoliagePaintingGeometryFilter& OverrideGeometryFilter, bool InRebuildFoliageTree = true);
 
 	/** Called when the user presses a button on their motion controller device */
 	void OnVRAction(class FEditorViewportClient& ViewportClient, class UViewportInteractor* Interactor, const struct FViewportActionKeyInput& Action, bool& bOutIsInputCaptured, bool& bWasHandled);
@@ -567,7 +567,7 @@ private:
 	void HandleOnFoliageTypeMeshChanged(UFoliageType* FoliageType);
 
 	/** Common code for adding instances to world based on settings */
-	static void AddInstancesImp(UWorld* InWorld, const UFoliageType* Settings, const TArray<FDesiredFoliageInstance>& DesiredInstances, const TArray<int32>& ExistingInstances = TArray<int32>(), const float Pressure = 1.f, LandscapeLayerCacheData* LandscapeLayerCaches = nullptr, const FFoliageUISettings* UISettings = nullptr, const FFoliagePaintingGeometryFilter* OverrideGeometryFilter = nullptr);
+	static void AddInstancesImp(UWorld* InWorld, const UFoliageType* Settings, const TArray<FDesiredFoliageInstance>& DesiredInstances, const TArray<int32>& ExistingInstances = TArray<int32>(), const float Pressure = 1.f, LandscapeLayerCacheData* LandscapeLayerCaches = nullptr, const FFoliageUISettings* UISettings = nullptr, const FFoliagePaintingGeometryFilter* OverrideGeometryFilter = nullptr, bool InRebuildFoliageTree = true);
 
 	/** Logic for determining which instances can be placed in the world*/
 	static void CalculatePotentialInstances(const UWorld* InWorld, const UFoliageType* Settings, const TArray<FDesiredFoliageInstance>& DesiredInstances, TArray<FPotentialInstance> OutPotentialInstances[NUM_INSTANCE_BUCKETS], LandscapeLayerCacheData* LandscaleLayerCachesPtr, const FFoliageUISettings* UISettings, const FFoliagePaintingGeometryFilter* OverrideGeometryFilter = nullptr);
diff --git a/Engine/Source/Editor/FoliageEdit/Private/ProceduralFoliageComponentDetails.cpp b/Engine/Source/Editor/FoliageEdit/Private/ProceduralFoliageComponentDetails.cpp
--- a/Engine/Source/Editor/FoliageEdit/Private/ProceduralFoliageComponentDetails.cpp
+++ b/Engine/Source/Editor/FoliageEdit/Private/ProceduralFoliageComponentDetails.cpp
@@ -91,7 +91,7 @@ FReply FProceduralFoliageComponentDetails::OnResimulateClicked()
 					OverrideGeometryFilter.bAllowFoliage = Component->bAllowFoliage;
 					OverrideGeometryFilter.bAllowTranslucent = Component->bAllowTranslucent;
 
-					FEdModeFoliage::AddInstances(Component->GetWorld(), DesiredFoliageInstances, OverrideGeometryFilter);
+					FEdModeFoliage::AddInstances(Component->GetWorld(), DesiredFoliageInstances, OverrideGeometryFilter, true);					
 				}
 
 				// If no instances were spawned, inform the user
diff --git a/Engine/Source/Runtime/Engine/Private/HierarchicalInstancedStaticMesh.cpp b/Engine/Source/Runtime/Engine/Private/HierarchicalInstancedStaticMesh.cpp
--- a/Engine/Source/Runtime/Engine/Private/HierarchicalInstancedStaticMesh.cpp
+++ b/Engine/Source/Runtime/Engine/Private/HierarchicalInstancedStaticMesh.cpp
@@ -1172,7 +1172,7 @@ void FHierarchicalStaticMeshSceneProxy::Traverse(const FFoliageCullInstanceParam
 
 	if (MinLOD != MaxLOD)
 	{
-		FVector ScaleAverage = (Node.MaxInstanceScale - Node.MinInstanceScale) / 2.0f;
+		FVector ScaleAverage = Node.MinInstanceScale + ((Node.MaxInstanceScale - Node.MinInstanceScale) / 2.0f);
 		float DistanceScale = FMath::Max(FMath::Max3(ScaleAverage.X, ScaleAverage.Y, ScaleAverage.Z), 0.001f);
 
 		CalcLOD(MinLOD, MaxLOD, Node.BoundMin, Node.BoundMax, Params.ViewOriginInLocalZero, Params.ViewOriginInLocalOne, Params.LODPlanesMin, Params.LODPlanesMax, DistanceScale);
