[NO NON-PII COMMIT MESSAGE]

diff --git a/Engine/Source/Editor/FoliageEdit/Private/FoliageEdMode.cpp b/Engine/Source/Editor/FoliageEdit/Private/FoliageEdMode.cpp
--- a/Engine/Source/Editor/FoliageEdit/Private/FoliageEdMode.cpp
+++ b/Engine/Source/Editor/FoliageEdit/Private/FoliageEdMode.cpp
@@ -1591,7 +1591,7 @@ void FEdModeFoliage::AddSingleInstanceForBrush(UWorld* InWorld, const UFoliageTy
 	TArray<int32> ExistingInstanceBuckets;
 	ExistingInstanceBuckets.AddZeroed(NUM_INSTANCE_BUCKETS);
 
-	AddInstancesImp(InWorld, Settings, DesiredInstances, ExistingInstanceBuckets, Pressure, &LandscapeLayerCaches, &UISettings, false);
+	AddInstancesImp(InWorld, Settings, DesiredInstances, ExistingInstanceBuckets, Pressure, &LandscapeLayerCaches, &UISettings, nullptr, false);
 }
 
 /** Add instances inside the brush to match DesiredInstanceCount */
@@ -1647,7 +1647,7 @@ void FEdModeFoliage::AddInstancesForBrush(UWorld* InWorld, const UFoliageType* S
 			FDesiredFoliageInstance* DesiredInstance = new (DesiredInstances)FDesiredFoliageInstance(Start, End);
 		}
 
-		AddInstancesImp(InWorld, Settings, DesiredInstances, ExistingInstanceBuckets, Pressure, &LandscapeLayerCaches, &UISettings, false);
+		AddInstancesImp(InWorld, Settings, DesiredInstances, ExistingInstanceBuckets, Pressure, &LandscapeLayerCaches, &UISettings, nullptr, false);
 	}
 }
 
