[NO NON-PII COMMIT MESSAGE]

diff --git a/Engine/Source/Editor/FoliageEdit/Private/ProceduralFoliageComponentDetails.cpp b/Engine/Source/Editor/FoliageEdit/Private/ProceduralFoliageComponentDetails.cpp
--- a/Engine/Source/Editor/FoliageEdit/Private/ProceduralFoliageComponentDetails.cpp
+++ b/Engine/Source/Editor/FoliageEdit/Private/ProceduralFoliageComponentDetails.cpp
@@ -82,7 +82,7 @@ FReply FProceduralFoliageComponentDetails::OnResimulateClicked()
 			{
 				if (DesiredFoliageInstances.Num() > 0)
 				{
-					Component->RemoveProceduralContent();
+					Component->RemoveProceduralContent(false);
 
 					FFoliagePaintingGeometryFilter OverrideGeometryFilter;
 					OverrideGeometryFilter.bAllowLandscape = Component->bAllowLandscape;
