Modified foliage thumbnail color to better show when it is inactive
#5057

diff --git a/Engine/Source/Editor/FoliageEdit/Private/FoliagePaletteItem.cpp b/Engine/Source/Editor/FoliageEdit/Private/FoliagePaletteItem.cpp
--- a/Engine/Source/Editor/FoliageEdit/Private/FoliagePaletteItem.cpp
+++ b/Engine/Source/Editor/FoliageEdit/Private/FoliagePaletteItem.cpp
@@ -472,8 +472,7 @@ void SFoliagePaletteItemTile::Construct(const FArguments& InArgs, TSharedRef<STa
 
 FLinearColor SFoliagePaletteItemTile::GetTileColorAndOpacity() const
 {
-	float Alpha = Model->IsActive() ? 1.f : 0.5f;
-	return FLinearColor(1.f, 1.f, 1.f, Alpha);
+	return Model->IsActive() ? FLinearColor::White : FLinearColor(0.5f, 0.5f, 0.5f, 1.f);
 }
 
 EVisibility SFoliagePaletteItemTile::GetCheckBoxVisibility() const
