Fix default asset name for actor foliage

#rnx

diff --git a/Engine/Source/Editor/FoliageEdit/Private/FoliageTypeFactory.h b/Engine/Source/Editor/FoliageEdit/Private/FoliageTypeFactory.h
--- a/Engine/Source/Editor/FoliageEdit/Private/FoliageTypeFactory.h
+++ b/Engine/Source/Editor/FoliageEdit/Private/FoliageTypeFactory.h
@@ -35,7 +35,7 @@ class UFoliageType_ActorFactory : public UFactory
 	virtual UObject* FactoryCreateNew(UClass* Class, UObject* InParent, FName Name, EObjectFlags Flags, UObject* Context, FFeedbackContext* Warn) override;
 	virtual FString GetDefaultNewAssetName() const override
 	{
-		return TEXT("NewBlueprintFoliage");
+		return TEXT("NewActorFoliage");
 	}
 	virtual FText GetToolTip() const override;
 	// End of UFactory interface
