FoliageEdMode: Move To Current Level Button

#rnx

diff --git a/Engine/Source/Editor/FoliageEdit/Private/SFoliageEdit.cpp b/Engine/Source/Editor/FoliageEdit/Private/SFoliageEdit.cpp
--- a/Engine/Source/Editor/FoliageEdit/Private/SFoliageEdit.cpp
+++ b/Engine/Source/Editor/FoliageEdit/Private/SFoliageEdit.cpp
@@ -27,6 +27,7 @@
 #include "Widgets/Layout/SHeader.h"
 #include "Widgets/Layout/SSeparator.h"
 #include "Widgets/Notifications/SErrorText.h"
+#include "Engine/World.h"
 
 #define LOCTEXT_NAMESPACE "FoliageEd_Mode"
 
@@ -478,6 +479,23 @@ void SFoliageEdit::Construct(const FArguments& InArgs)
 								.ToolTipText(LOCTEXT("DeselectAllInstances_Tooltip", "Deselects all foliage instances"))
 							]
 						]
+
+						// Move to Current Level
+						+ SWrapBox::Slot()
+						.Padding(FMargin(0.f, 0.f, 6.f, 3.f))
+						[
+							SNew(SBox)
+							.WidthOverride(150.f)
+							.HeightOverride(25.f)
+							[
+								SNew(SButton)
+								.HAlign(HAlign_Center)
+								.VAlign(VAlign_Center)
+								.OnClicked(this, &SFoliageEdit::OnMoveSelectedInstancesToCurrentLevel)
+								.Text(LOCTEXT("MoveSelectedInstancesToCurrentLevel", "Move to Current Level"))
+								.ToolTipText(LOCTEXT("MoveSelectedInstancesToCurrentLevel_Tooltip", "Move selected foliage instances to current level"))
+							]
+						]
 					]
 				]
 			]
@@ -709,6 +727,19 @@ FReply SFoliageEdit::OnDeselectAllInstances()
 	return FReply::Handled();
 }
 
+FReply SFoliageEdit::OnMoveSelectedInstancesToCurrentLevel()
+{
+	if (UWorld* World = FoliageEditMode->GetWorld())
+	{
+		if (ULevel* CurrentLevel = World->GetCurrentLevel())
+		{
+			FoliageEditMode->MoveSelectedFoliageToLevel(CurrentLevel);
+		}
+	}
+
+	return FReply::Handled();
+}
+
 FText SFoliageEdit::GetFilterText() const
 {
 	FText TooltipText;
diff --git a/Engine/Source/Editor/FoliageEdit/Private/SFoliageEdit.h b/Engine/Source/Editor/FoliageEdit/Private/SFoliageEdit.h
--- a/Engine/Source/Editor/FoliageEdit/Private/SFoliageEdit.h
+++ b/Engine/Source/Editor/FoliageEdit/Private/SFoliageEdit.h
@@ -194,6 +194,9 @@ private:	// SELECTION
 	/** Handler for 'Deselect All' command  */
 	FReply OnDeselectAllInstances();
 
+	/** Handler for 'Move to Current Level' command*/
+	FReply OnMoveSelectedInstancesToCurrentLevel();
+
 	/** Tooltip text for 'Instance Count" column */
 	FText GetTotalInstanceCountTooltipText() const;
 
