Fix Performance issue when multi-selecting Foliage Actors (Only notify editor once)

#rnx

diff --git a/Engine/Source/Editor/FoliageEdit/Private/FoliageEdMode.cpp b/Engine/Source/Editor/FoliageEdit/Private/FoliageEdMode.cpp
--- a/Engine/Source/Editor/FoliageEdit/Private/FoliageEdMode.cpp
+++ b/Engine/Source/Editor/FoliageEdit/Private/FoliageEdMode.cpp
@@ -735,11 +735,19 @@ void FEdModeFoliage::NotifyAssetRemoved(const FAssetData& AssetInfo)
 
 void FEdModeFoliage::NotifyActorSelectionChanged(bool bSelect, const TArray<AActor*>& Selection)
 {
+	if (Selection.Num() == 0)
+	{
+		return;
+	}
+
 	GEditor->GetSelectedActors()->Modify();
 	for (AActor* Actor : Selection)
 	{
-		GEditor->SelectActor(Actor, bSelect, true, true);
+        const bool bNotify = false;
+		const bool bSelectEvenIfHidden = true;
+		GEditor->SelectActor(Actor, bSelect, bNotify, bSelectEvenIfHidden);
 	}
+	GEditor->NoteSelectionChange();
 }
 
 /** When the user changes the current tool in the UI */
