Fixed 'move selected instanced foliage actors to level'. Now moves only the selection, no all instanced foliage actors.

diff --git a/Engine/Source/Editor/FoliageEdit/Private/FoliageEditUtility.cpp b/Engine/Source/Editor/FoliageEdit/Private/FoliageEditUtility.cpp
--- a/Engine/Source/Editor/FoliageEdit/Private/FoliageEditUtility.cpp
+++ b/Engine/Source/Editor/FoliageEdit/Private/FoliageEditUtility.cpp
@@ -120,7 +120,7 @@ void FFoliageEditUtility::ReplaceFoliageTypeObject(UWorld* InWorld, UFoliageType
 }
 
 
-void FFoliageEditUtility::MoveActorFoliageInstancesToLevel(ULevel* InTargetLevel)
+void FFoliageEditUtility::MoveActorFoliageInstancesToLevel(ULevel* InTargetLevel, AActor* InIFA)
 {
 	// Can't move into a locked level
 	if (FLevelUtils::IsLevelLocked(InTargetLevel))
@@ -151,6 +151,10 @@ void FFoliageEditUtility::MoveActorFoliageInstancesToLevel(ULevel* InTargetLevel
 			{
 				continue;
 			}
+			if (InIFA && IFA != InIFA)
+			{
+				continue;
+			}
 
 			bool CanMoveInstanceType = true;
 
@@ -196,6 +200,11 @@ void FFoliageEditUtility::MoveActorFoliageInstancesToLevel(ULevel* InTargetLevel
 
 				IFA->MoveAllInstancesToLevel(InTargetLevel);
 			}
+
+			if (InIFA)
+			{
+				return;
+			}
 		}
 	}
 }
diff --git a/Engine/Source/Editor/FoliageEdit/Public/FoliageEditUtility.h b/Engine/Source/Editor/FoliageEdit/Public/FoliageEditUtility.h
--- a/Engine/Source/Editor/FoliageEdit/Public/FoliageEditUtility.h
+++ b/Engine/Source/Editor/FoliageEdit/Public/FoliageEditUtility.h
@@ -20,5 +20,5 @@ public:
 	static FOLIAGEEDIT_API UFoliageType* SaveFoliageTypeObject(UFoliageType* Settings);
 	static FOLIAGEEDIT_API void ReplaceFoliageTypeObject(UWorld* InWorld, UFoliageType* OldType, UFoliageType* NewType);
 
-	static FOLIAGEEDIT_API void MoveActorFoliageInstancesToLevel(ULevel* InTargetLevel);
+	static FOLIAGEEDIT_API void MoveActorFoliageInstancesToLevel(ULevel* InTargetLevel, AActor* InIFA = nullptr);
 };
\ No newline at end of file
diff --git a/Engine/Source/Editor/WorldBrowser/Private/LevelCollectionModel.cpp b/Engine/Source/Editor/WorldBrowser/Private/LevelCollectionModel.cpp
--- a/Engine/Source/Editor/WorldBrowser/Private/LevelCollectionModel.cpp
+++ b/Engine/Source/Editor/WorldBrowser/Private/LevelCollectionModel.cpp
@@ -1732,7 +1732,7 @@ void FLevelCollectionModel::MoveActorsToSelected_Executed()
 
 		if (Actor != nullptr)
 		{
-			FFoliageEditUtility::MoveActorFoliageInstancesToLevel(GetWorld()->GetCurrentLevel());
+			FFoliageEditUtility::MoveActorFoliageInstancesToLevel(GetWorld()->GetCurrentLevel(), Actor);
 		}
 	}
 
