Fixed editor build
#rnx

diff --git a/Engine/Source/Editor/FoliageEdit/Public/FoliageEditUtility.h b/Engine/Source/Editor/FoliageEdit/Public/FoliageEditUtility.h
--- a/Engine/Source/Editor/FoliageEdit/Public/FoliageEditUtility.h
+++ b/Engine/Source/Editor/FoliageEdit/Public/FoliageEditUtility.h
@@ -12,6 +12,7 @@
 class UFoliageType;
 struct FFoliageMeshUIInfo;
 class ULevel;
+class AActor;
 
 class FFoliageEditUtility
 {
