Foliage Selection is now Transient (not saved anymore) so it doesn't dirty the Foliage Type

diff --git a/Engine/Source/Editor/FoliageEdit/Private/FoliagePaletteItem.cpp b/Engine/Source/Editor/FoliageEdit/Private/FoliagePaletteItem.cpp
--- a/Engine/Source/Editor/FoliageEdit/Private/FoliagePaletteItem.cpp
+++ b/Engine/Source/Editor/FoliageEdit/Private/FoliagePaletteItem.cpp
@@ -303,12 +303,12 @@ void FFoliagePaletteItemModel::SetTypeActiveInPalette(bool bSetActiveInPalette)
 	UFoliageType* FoliageType = GetFoliageType();
 	if (FoliageType->IsSelected != bSetActiveInPalette)
 	{
-		FoliageType->Modify();
+		FoliageType->Modify(false);
 		FoliageType->IsSelected = bSetActiveInPalette;
 
 		if (IsBlueprint())
 		{
-			FoliageType->GetClass()->ClassGeneratedBy->Modify();
+			FoliageType->GetClass()->ClassGeneratedBy->Modify(false);
 			FoliageType->GetClass()->GetDefaultObject<UFoliageType>()->IsSelected = bSetActiveInPalette;
 		}
 	}
diff --git a/Engine/Source/Runtime/Foliage/Public/FoliageType.h b/Engine/Source/Runtime/Foliage/Public/FoliageType.h
--- a/Engine/Source/Runtime/Foliage/Public/FoliageType.h
+++ b/Engine/Source/Runtime/Foliage/Public/FoliageType.h
@@ -334,7 +334,7 @@ public:
 	UPROPERTY(transient)
 	uint64 HiddenEditorViews;
 
-	UPROPERTY()
+	UPROPERTY(transient)
 	uint32 IsSelected:1;
 #endif
 public:
