Foliage: New option in single instance placement to cycle through selected foliage instead of placing all the foliage types at the same time

diff --git a/Engine/Source/Editor/FoliageEdit/Private/FoliageEdMode.cpp b/Engine/Source/Editor/FoliageEdit/Private/FoliageEdMode.cpp
--- a/Engine/Source/Editor/FoliageEdit/Private/FoliageEdMode.cpp
+++ b/Engine/Source/Editor/FoliageEdit/Private/FoliageEdMode.cpp
@@ -31,6 +31,7 @@
 #include "FoliageEditActions.h"
 
 #include "AssetRegistryModule.h"
+#include "Misc/ScopeExit.h"
 
 //Slate dependencies
 #include "ILevelViewport.h"
@@ -1494,13 +1495,13 @@ void FEdModeFoliage::RebuildFoliageTree(const UFoliageType* Settings)
 	}
 }
 
-void FEdModeFoliage::AddInstancesImp(UWorld* InWorld, const UFoliageType* Settings, const TArray<FDesiredFoliageInstance>& DesiredInstances, const TArray<int32>& ExistingInstanceBuckets, const float Pressure, LandscapeLayerCacheData* LandscapeLayerCachesPtr, const FFoliageUISettings* UISettings, const FFoliagePaintingGeometryFilter* OverrideGeometryFilter, bool InRebuildFoliageTree)
+bool FEdModeFoliage::AddInstancesImp(UWorld* InWorld, const UFoliageType* Settings, const TArray<FDesiredFoliageInstance>& DesiredInstances, const TArray<int32>& ExistingInstanceBuckets, const float Pressure, LandscapeLayerCacheData* LandscapeLayerCachesPtr, const FFoliageUISettings* UISettings, const FFoliagePaintingGeometryFilter* OverrideGeometryFilter, bool InRebuildFoliageTree)
 {
 	SCOPE_CYCLE_COUNTER(STAT_FoliageAddInstanceImp);
 	
 	if (DesiredInstances.Num() == 0)
 	{
-		return;
+		return false;
 	}
 
 	TArray<FPotentialInstance> PotentialInstanceBuckets[NUM_INSTANCE_BUCKETS];
@@ -1533,6 +1534,8 @@ void FEdModeFoliage::AddInstancesImp(UWorld* InWorld, const UFoliageType* Settin
 		}
 	}
 
+	bool bPlacedInstances = false;
+
 	for (int32 BucketIdx = 0; BucketIdx < NUM_INSTANCE_BUCKETS; BucketIdx++)
 	{
 		TArray<FPotentialInstance>& PotentialInstances = PotentialInstanceBuckets[BucketIdx];
@@ -1558,15 +1561,18 @@ void FEdModeFoliage::AddInstancesImp(UWorld* InWorld, const UFoliageType* Settin
 					Inst.ProceduralGuid = PotentialInstance.DesiredInstance.ProceduralGuid;
 					Inst.BaseComponent = PotentialInstance.HitComponent;
 					PlacedInstances.Add(MoveTemp(Inst));
+					bPlacedInstances = true;
 				}
 			}
 
 			SpawnFoliageInstance(InWorld, Settings, UISettings, PlacedInstances, InRebuildFoliageTree);
 		}
 	}
+
+	return bPlacedInstances;
 }
 
-void FEdModeFoliage::AddSingleInstanceForBrush(UWorld* InWorld, const UFoliageType* Settings, float Pressure)
+bool FEdModeFoliage::AddSingleInstanceForBrush(UWorld* InWorld, const UFoliageType* Settings, float Pressure)
 {
 	SCOPE_CYCLE_COUNTER(STAT_FoliageAddInstanceBrush);
 
@@ -1583,7 +1589,7 @@ void FEdModeFoliage::AddSingleInstanceForBrush(UWorld* InWorld, const UFoliageTy
 	TArray<int32> ExistingInstanceBuckets;
 	ExistingInstanceBuckets.AddZeroed(NUM_INSTANCE_BUCKETS);
 
-	AddInstancesImp(InWorld, Settings, DesiredInstances, ExistingInstanceBuckets, Pressure, &LandscapeLayerCaches, &UISettings, nullptr, false);
+	return AddInstancesImp(InWorld, Settings, DesiredInstances, ExistingInstanceBuckets, Pressure, &LandscapeLayerCaches, &UISettings, nullptr, false);
 }
 
 /** Add instances inside the brush to match DesiredInstanceCount */
@@ -2449,15 +2455,24 @@ void FEdModeFoliage::ApplyBrush(FEditorViewportClient* ViewportClient)
 
 	// Cache a copy of the world pointer
 	UWorld* World = ViewportClient->GetWorld();
-
+	int32 SelectedIndex = -1;
+	TArray<FFoliageMeshUIInfoPtr> SelectedFoliageMeshList;
+	SelectedFoliageMeshList.Reserve(FoliageMeshList.Num());
 	for (auto& FoliageMeshUI : FoliageMeshList)
 	{
-		UFoliageType* Settings = FoliageMeshUI->Settings;
-
-		if (!Settings->IsSelected)
+		if (FoliageMeshUI->Settings->IsSelected)
 		{
-			continue;
+			SelectedFoliageMeshList.Add(FoliageMeshUI);
 		}
+	}
+
+	for (int32 Index = 0; Index < SelectedFoliageMeshList.Num(); ++Index)
+	{
+		UFoliageType* Settings = SelectedFoliageMeshList[Index]->Settings;
+		ON_SCOPE_EXIT
+		{
+			OnInstanceCountUpdated(Settings);
+		};
 
 		FSphere BrushSphere(BrushLocation, UISettings.GetRadius());
 
@@ -2482,7 +2497,21 @@ void FEdModeFoliage::ApplyBrush(FEditorViewportClient* ViewportClient)
 			{
 				if (UISettings.IsInAnySingleInstantiationMode())
 				{
-					AddSingleInstanceForBrush(World, Settings, Pressure);
+					if (UISettings.GetSingleInstantiationPlacementMode() == EFoliageSingleInstantiationPlacementMode::Type::All)
+					{
+						AddSingleInstanceForBrush(World, Settings, Pressure);
+					}
+					else if (UISettings.GetSingleInstantiationPlacementMode() == EFoliageSingleInstantiationPlacementMode::Type::CycleThrough)
+					{
+						if (UISettings.GetSingleInstantiationCycleThroughIndex() % SelectedFoliageMeshList.Num() == Index)
+						{
+							if (AddSingleInstanceForBrush(World, Settings, Pressure))
+							{
+								UISettings.IncrementSingleInstantiationCycleThroughIndex();
+							}
+							return;
+						}
+					}
 				}
 				else
 				{
@@ -2495,8 +2524,6 @@ void FEdModeFoliage::ApplyBrush(FEditorViewportClient* ViewportClient)
 				}
 			}
 		}
-
-		OnInstanceCountUpdated(Settings);
 	}
 
 	if (UISettings.GetLassoSelectToolSelected())
@@ -3788,6 +3815,10 @@ void FFoliageUISettings::Load()
 
 	GConfig->GetBool(TEXT("FoliageEdit"), TEXT("IsInSingleInstantiationMode"), IsInSingleInstantiationMode, GEditorPerProjectIni);
 	GConfig->GetBool(TEXT("FoliageEdit"), TEXT("IsInSpawnInCurrentLevelMode"), IsInSpawnInCurrentLevelMode, GEditorPerProjectIni);
+
+	int32 SingleInstantiationPlacementModeAsInt = 0;
+	GConfig->GetInt(TEXT("FoliageEdit"), TEXT("SingleInstantiationPlacementMode"), SingleInstantiationPlacementModeAsInt, GEditorPerProjectIni);
+	SingleInstantiationPlacementMode = EFoliageSingleInstantiationPlacementMode::Type(SingleInstantiationPlacementModeAsInt);
 }
 
 /** Save UI settings to ini file */
@@ -3812,7 +3843,7 @@ void FFoliageUISettings::Save()
 
 	GConfig->SetBool(TEXT("FoliageEdit"), TEXT("IsInSingleInstantiationMode"), IsInSingleInstantiationMode, GEditorPerProjectIni);
 	GConfig->SetBool(TEXT("FoliageEdit"), TEXT("IsInSpawnInCurrentLevelMode"), IsInSpawnInCurrentLevelMode, GEditorPerProjectIni);
-
+	GConfig->SetInt(TEXT("FoliageEdit"), TEXT("SingleInstantiationPlacementMode"), (int32)SingleInstantiationPlacementMode, GEditorPerProjectIni);
 }
 
 #undef LOCTEXT_NAMESPACE
diff --git a/Engine/Source/Editor/FoliageEdit/Private/FoliageEdMode.h b/Engine/Source/Editor/FoliageEdit/Private/FoliageEdMode.h
--- a/Engine/Source/Editor/FoliageEdit/Private/FoliageEdMode.h
+++ b/Engine/Source/Editor/FoliageEdit/Private/FoliageEdMode.h
@@ -38,6 +38,17 @@ namespace EFoliagePaletteViewMode
 	};
 }
 
+/** Single instance mode */
+namespace EFoliageSingleInstantiationPlacementMode
+{
+	enum class Type
+	{
+		All,
+		CycleThrough,
+		ModeCount
+	};
+};
+
 const float SingleInstanceModeBrushSize = 20.0f;
 
 // Current user settings in Foliage UI
@@ -89,6 +100,12 @@ struct FFoliageUISettings
 	bool GetIsInQuickSingleInstantiationMode() const { return IsInQuickSingleInstantiationMode; }
 	void SetIsInQuickSingleInstantiationMode(bool InIsInQuickSingleInstantiationMode) { IsInQuickSingleInstantiationMode = InIsInQuickSingleInstantiationMode; }
 
+	EFoliageSingleInstantiationPlacementMode::Type GetSingleInstantiationPlacementMode() const { return SingleInstantiationPlacementMode; }
+	void SetSingleInstantiationPlacementMode(EFoliageSingleInstantiationPlacementMode::Type InSingleInstantiationPlacementMode) { SingleInstantiationPlacementMode = InSingleInstantiationPlacementMode; }
+
+	int32 GetSingleInstantiationCycleThroughIndex() const { return SingleInstantiationCycleThroughIndex; }
+	void IncrementSingleInstantiationCycleThroughIndex() { SingleInstantiationCycleThroughIndex++; }
+
 	bool GetIsInSpawnInCurrentLevelMode() const { return IsInSpawnInCurrentLevelMode; }
 	void SetSpawnInCurrentLevelMode(bool InSpawnInCurrentLevelMode) { IsInSpawnInCurrentLevelMode = InSpawnInCurrentLevelMode; }
 
@@ -121,6 +138,8 @@ struct FFoliageUISettings
 		, UnpaintDensity(0.f)
 		, IsInSingleInstantiationMode(false)
 		, IsInQuickSingleInstantiationMode(false)
+		, SingleInstantiationPlacementMode(EFoliageSingleInstantiationPlacementMode::Type::All)
+		, SingleInstantiationCycleThroughIndex(0)
 		, IsInSpawnInCurrentLevelMode(false)
 		, bFilterLandscape(true)
 		, bFilterStaticMesh(true)
@@ -158,6 +177,8 @@ private:
 
 	bool IsInSingleInstantiationMode;
 	bool IsInQuickSingleInstantiationMode;
+	EFoliageSingleInstantiationPlacementMode::Type SingleInstantiationPlacementMode;
+	int32 SingleInstantiationCycleThroughIndex;
 	bool IsInSpawnInCurrentLevelMode;
 
 public:
@@ -524,7 +545,10 @@ private:
 	/** Add instances inside the brush to match DesiredInstanceCount */
 	void AddInstancesForBrush(UWorld* InWorld, const UFoliageType* Settings, const FSphere& BrushSphere, int32 DesiredInstanceCount, float Pressure);
 
-	void AddSingleInstanceForBrush(UWorld* InWorld, const UFoliageType* Settings, float Pressure);
+	/** Add single instance inside the brush 
+	* @return true if instance was added successfully
+	*/
+	bool AddSingleInstanceForBrush(UWorld* InWorld, const UFoliageType* Settings, float Pressure);
 
 	/** Remove instances inside the brush to match DesiredInstanceCount. */
 	void RemoveInstancesForBrush(UWorld* InWorld, const UFoliageType* Settings, const FSphere& BrushSphere, int32 DesiredInstanceCount, float Pressure);
@@ -576,7 +600,7 @@ private:
 	void HandleOnFoliageTypeMeshChanged(UFoliageType* FoliageType);
 
 	/** Common code for adding instances to world based on settings */
-	static void AddInstancesImp(UWorld* InWorld, const UFoliageType* Settings, const TArray<FDesiredFoliageInstance>& DesiredInstances, const TArray<int32>& ExistingInstances = TArray<int32>(), const float Pressure = 1.f, LandscapeLayerCacheData* LandscapeLayerCaches = nullptr, const FFoliageUISettings* UISettings = nullptr, const FFoliagePaintingGeometryFilter* OverrideGeometryFilter = nullptr, bool InRebuildFoliageTree = true);
+	static bool AddInstancesImp(UWorld* InWorld, const UFoliageType* Settings, const TArray<FDesiredFoliageInstance>& DesiredInstances, const TArray<int32>& ExistingInstances = TArray<int32>(), const float Pressure = 1.f, LandscapeLayerCacheData* LandscapeLayerCaches = nullptr, const FFoliageUISettings* UISettings = nullptr, const FFoliagePaintingGeometryFilter* OverrideGeometryFilter = nullptr, bool InRebuildFoliageTree = true);
 
 	/** Logic for determining which instances can be placed in the world*/
 	static void CalculatePotentialInstances(const UWorld* InWorld, const UFoliageType* Settings, const TArray<FDesiredFoliageInstance>& DesiredInstances, TArray<FPotentialInstance> OutPotentialInstances[NUM_INSTANCE_BUCKETS], LandscapeLayerCacheData* LandscaleLayerCachesPtr, const FFoliageUISettings* UISettings, const FFoliagePaintingGeometryFilter* OverrideGeometryFilter = nullptr);
diff --git a/Engine/Source/Editor/FoliageEdit/Private/SFoliageEdit.cpp b/Engine/Source/Editor/FoliageEdit/Private/SFoliageEdit.cpp
--- a/Engine/Source/Editor/FoliageEdit/Private/SFoliageEdit.cpp
+++ b/Engine/Source/Editor/FoliageEdit/Private/SFoliageEdit.cpp
@@ -230,6 +230,47 @@ void SFoliageEdit::Construct(const FArguments& InArgs)
 						]
 					]					
 					
+					+ SVerticalBox::Slot()
+					.Padding(StandardPadding)
+					.AutoHeight()
+					[
+						SNew(SHorizontalBox)
+						.Visibility(this, &SFoliageEdit::GetVisibility_Options)
+						+ SHorizontalBox::Slot()
+						.VAlign(VAlign_Center)
+						.MaxWidth(140)
+						.Padding(StandardLeftPadding)
+						[
+							SNew(SCheckBox)
+							.Visibility(this, &SFoliageEdit::GetVisibility_SingleInstantiationMode)
+							.OnCheckStateChanged(this, &SFoliageEdit::OnCheckStateChanged_SingleInstantiationMode)
+							.IsChecked(this, &SFoliageEdit::GetCheckState_SingleInstantiationMode)
+							.ToolTipText(LOCTEXT("SingleInstantiationModeTooltips", "Paint a single foliage instance at the mouse cursor location (i + Mouse Click)"))
+							[
+								SNew(STextBlock)
+								.Text(LOCTEXT("SingleInstantiationMode", "Single Instance Mode: "))
+								.Font(StandardFont)
+							]
+						]
+						+ SHorizontalBox::Slot()
+						.VAlign(VAlign_Center)
+						.HAlign(HAlign_Left)
+						.Padding(StandardRightPadding)
+						[
+							SNew(SComboButton)
+							.Visibility(this, &SFoliageEdit::GetVisibility_SingleInstantiationPlacementMode)
+							.IsEnabled(this, &SFoliageEdit::GetIsEnabled_SingleInstantiationPlacementMode)
+							.OnGetMenuContent(this, &SFoliageEdit::GetSingleInstantiationModeMenuContent)
+							.ContentPadding(2)
+							.ToolTipText(LOCTEXT("SingleInstantiationPlacementModeToolTips", "Changes the placement mode when using single instance"))
+							.ButtonContent()
+							[
+								SNew(STextBlock)
+								.Text(this, &SFoliageEdit::GetCurrentSingleInstantiationPlacementModeText)
+							]
+						]
+					]
+
 					+ SVerticalBox::Slot()
 					.Padding(StandardPadding)
 					.AutoHeight()
@@ -245,24 +286,6 @@ void SFoliageEdit::Construct(const FArguments& InArgs)
 							.UseAllottedWidth(true)
 							.InnerSlotPadding({6, 5})
 
-							+ SWrapBox::Slot()
-							[
-								SNew(SBox)
-								.MinDesiredWidth(150)
-								[
-									SNew(SCheckBox)
-									.Visibility(this, &SFoliageEdit::GetVisibility_SingleInstantiationMode)
-									.OnCheckStateChanged(this, &SFoliageEdit::OnCheckStateChanged_SingleInstantiationMode)
-									.IsChecked(this, &SFoliageEdit::GetCheckState_SingleInstantiationMode)
-									.ToolTipText(LOCTEXT("SingleInstantiationModeTooltips", "Paint a single foliage instance at the mouse cursor location (i + Mouse Click)"))
-									[
-										SNew(STextBlock)
-										.Text(LOCTEXT("SingleInstantiationMode", "Single Instance Mode"))
-										.Font(StandardFont)
-									]
-								]
-							]
-
 							+ SWrapBox::Slot()
 							[
 								SNew(SBox)
@@ -995,6 +1018,21 @@ EVisibility SFoliageEdit::GetVisibility_SingleInstantiationMode() const
 	return EVisibility::Collapsed;
 }
 
+EVisibility SFoliageEdit::GetVisibility_SingleInstantiationPlacementMode() const
+{
+	if (FoliageEditMode->UISettings.GetPaintToolSelected())
+	{
+		return EVisibility::Visible;
+	}
+
+	return EVisibility::Collapsed;
+}
+
+bool SFoliageEdit::GetIsEnabled_SingleInstantiationPlacementMode() const
+{
+	return FoliageEditMode->UISettings.IsInAnySingleInstantiationMode();
+}
+
 EVisibility SFoliageEdit::GetVisibility_SpawnInCurrentLevelMode() const
 {
 	if (FoliageEditMode->UISettings.GetPaintToolSelected())
@@ -1015,4 +1053,39 @@ EVisibility SFoliageEdit::GetVisibility_Options() const
 	return EVisibility::Visible;
 }
 
+void SFoliageEdit::OnSingleInstantiationPlacementModeChanged(int32 InMode)
+{
+	FoliageEditMode->UISettings.SetSingleInstantiationPlacementMode(EFoliageSingleInstantiationPlacementMode::Type(InMode));
+}
+
+TSharedRef<SWidget> SFoliageEdit::GetSingleInstantiationModeMenuContent() const
+{
+	FMenuBuilder MenuBuilder(true, nullptr);
+
+	for (int32 i = 0; i < (int32)EFoliageSingleInstantiationPlacementMode::Type::ModeCount; i++)
+	{
+		MenuBuilder.AddMenuEntry(GetSingleInstantiationPlacementModeText(EFoliageSingleInstantiationPlacementMode::Type(i)), FText(), FSlateIcon(), FExecuteAction::CreateSP(this, &SFoliageEdit::OnSingleInstantiationPlacementModeChanged, i));
+	}
+
+	return MenuBuilder.MakeWidget();
+}
+
+FText SFoliageEdit::GetSingleInstantiationPlacementModeText(EFoliageSingleInstantiationPlacementMode::Type InMode) const
+{
+	switch (InMode)
+	{
+	case EFoliageSingleInstantiationPlacementMode::Type::All:
+		return LOCTEXT("SingleInstantiationPlacementModeAll", "All Selected");
+	case EFoliageSingleInstantiationPlacementMode::Type::CycleThrough:
+		return LOCTEXT("SingleInstantiationPlacementModeCycleThrough", "Cycle Through Selected");
+	}
+
+	return LOCTEXT("SingleInstantiationPlacementModeNone", "Invalid");
+}
+
+FText SFoliageEdit::GetCurrentSingleInstantiationPlacementModeText() const
+{
+	return GetSingleInstantiationPlacementModeText(FoliageEditMode->UISettings.GetSingleInstantiationPlacementMode());
+}
+
 #undef LOCTEXT_NAMESPACE
diff --git a/Engine/Source/Editor/FoliageEdit/Private/SFoliageEdit.h b/Engine/Source/Editor/FoliageEdit/Private/SFoliageEdit.h
--- a/Engine/Source/Editor/FoliageEdit/Private/SFoliageEdit.h
+++ b/Engine/Source/Editor/FoliageEdit/Private/SFoliageEdit.h
@@ -14,6 +14,9 @@ class SFoliagePalette;
 class UFoliageType;
 struct FFoliageMeshUIInfo;
 enum class ECheckBoxState : uint8;
+namespace EFoliageSingleInstantiationPlacementMode {
+	enum class Type;
+};
 
 typedef TSharedPtr<FFoliageMeshUIInfo> FFoliageMeshUIInfoPtr; //should match typedef in FoliageEdMode.h
 
@@ -180,8 +183,26 @@ private:	// BRUSH SETTINGS
 	/** Checks if the SingleInstantiationMode checkbox should appear. Dependant on the current tool being used. */
 	EVisibility GetVisibility_SingleInstantiationMode() const;	
 
+	/** Checks if the SingleInstantiationPlacement mode combobutton should appear. Dependant on the current tool being used. */
+	EVisibility GetVisibility_SingleInstantiationPlacementMode() const;
+
+	/** Checks if the SingleInstantiationPlacement mode combobutton should be enabled. Dependant on if single instantiation mode is enabled. */
+	bool GetIsEnabled_SingleInstantiationPlacementMode() const;
+
 	/** Checks if the SpawnInCurrentLevelMode checkbox should appear. Dependant on the current tool being used. */
 	EVisibility GetVisibility_SpawnInCurrentLevelMode() const;
+		
+	/** Sets the single instantiation placement mode */
+	void OnSingleInstantiationPlacementModeChanged(int32 InMode);
+
+	/** @return display text for the placement type */
+	FText GetSingleInstantiationPlacementModeText(EFoliageSingleInstantiationPlacementMode::Type InMode) const;
+
+	/** @return menu of all single instantiation modes */
+	TSharedRef<SWidget> GetSingleInstantiationModeMenuContent() const;
+
+	/** @return display text for current placement type */
+	FText GetCurrentSingleInstantiationPlacementModeText() const;
 
 private:	// SELECTION
 
@@ -210,7 +231,7 @@ private:
 	
 	/** Current error message */	
 	TSharedPtr<class SErrorText> ErrorText;
-
+		
 	/** Pointer to the foliage edit mode. */
 	FEdModeFoliage*					FoliageEditMode;
 };
