- Added option to override radius used to detect collision with other instances when painting in single instance mode
- Single instance mode painting sphere brush is now displaying maximum radius of selected Foliage Types to paint (also takes into consideration override radius if enabled)
- Single instance mode painting sphere brush color is green to differentiate from the normal foliage painting

diff --git a/Engine/Source/Editor/FoliageEdit/Private/FoliageEdMode.cpp b/Engine/Source/Editor/FoliageEdit/Private/FoliageEdMode.cpp
--- a/Engine/Source/Editor/FoliageEdit/Private/FoliageEdMode.cpp
+++ b/Engine/Source/Editor/FoliageEdit/Private/FoliageEdMode.cpp
@@ -199,6 +199,8 @@ bool FFoliagePaintingGeometryFilter::operator() (const UPrimitiveComponent* Comp
 // FEdModeFoliage
 //
 
+static FName FoliageBrushHighlightColorParamName("HighlightColor");
+
 /** Constructor */
 FEdModeFoliage::FEdModeFoliage()
 	: FEdMode()
@@ -210,15 +212,18 @@ FEdModeFoliage::FEdModeFoliage()
 {
 	// Load resources and construct brush component
 	UStaticMesh* StaticMesh = nullptr;
+	BrushDefaultHighlightColor = FColor::White;
 	if (!IsRunningCommandlet())
 	{
 		UMaterial* BrushMaterial = LoadObject<UMaterial>(nullptr, TEXT("/Engine/EditorLandscapeResources/FoliageBrushSphereMaterial.FoliageBrushSphereMaterial"), nullptr, LOAD_None, nullptr);
 		BrushMID = UMaterialInstanceDynamic::Create(BrushMaterial, GetTransientPackage());
 		check(BrushMID != nullptr);
-
+		FLinearColor DefaultColor;
+		BrushMID->GetVectorParameterDefaultValue(FoliageBrushHighlightColorParamName, DefaultColor);
+		BrushDefaultHighlightColor = DefaultColor.ToFColor(false);
 		StaticMesh = LoadObject<UStaticMesh>(nullptr, TEXT("/Engine/EngineMeshes/Sphere.Sphere"), nullptr, LOAD_None, nullptr);
 	}
-
+	BrushCurrentHighlightColor = BrushDefaultHighlightColor;
 	SphereBrushComponent = NewObject<UStaticMeshComponent>(GetTransientPackage(), TEXT("SphereBrushComponent"));
 	SphereBrushComponent->SetCollisionProfileName(UCollisionProfile::NoCollision_ProfileName);
 	SphereBrushComponent->SetCollisionObjectType(ECC_WorldDynamic);
@@ -883,9 +888,17 @@ void FEdModeFoliage::Tick(FEditorViewportClient* ViewportClient, float DeltaTime
 	if (bBrushTraceValid && (UISettings.GetPaintToolSelected() || UISettings.GetReapplyToolSelected() || UISettings.GetLassoSelectToolSelected()))
 	{
 		// Scale adjustment is due to default sphere SM size.
-		FTransform BrushTransform = FTransform(FQuat::Identity, BrushLocation, FVector(UISettings.GetRadius() * 0.00625f));
+		FTransform BrushTransform = FTransform(FQuat::Identity, BrushLocation, FVector(GetPaintingBrushRadius() * 0.00625f));
 		SphereBrushComponent->SetRelativeTransform(BrushTransform);
 
+		static FColor BrushSingleInstanceModeHighlightColor = FColor::Green;
+		FColor BrushHighlightColor = UISettings.IsInAnySingleInstantiationMode() ? BrushSingleInstanceModeHighlightColor : BrushDefaultHighlightColor;
+		if (BrushCurrentHighlightColor != BrushHighlightColor)
+		{
+			BrushCurrentHighlightColor = BrushHighlightColor;
+			BrushMID->SetVectorParameterValue(FoliageBrushHighlightColorParamName, BrushHighlightColor);
+		}
+		
 		if (!SphereBrushComponent->IsRegistered())
 		{
 			SphereBrushComponent->RegisterComponentWithWorld(ViewportClient->GetWorld());
@@ -1128,25 +1141,27 @@ static bool CheckForOverlappingSphere(const UWorld* InWorld, const UFoliageType*
 	return false;
 }
 
-static bool CheckLocationForPotentialInstance(const UWorld* InWorld, const UFoliageType* Settings, const FVector& Location, const FVector& Normal, TArray<FVector>& PotentialInstanceLocations, FFoliageInstanceHash& PotentialInstanceHash)
+static bool CheckLocationForPotentialInstance(const UWorld* InWorld, const UFoliageType* Settings, const bool bSingleInstanceMode, const FVector& Location, const FVector& Normal, TArray<FVector>& PotentialInstanceLocations, FFoliageInstanceHash& PotentialInstanceHash)
 {
 	if (CheckLocationForPotentialInstance_ThreadSafe(Settings, Location, Normal) == false)
 	{
 		return false;
 	}
 
+	const float SettingsRadius = Settings->GetRadius(bSingleInstanceMode);
+
 	// Check if we're too close to any other instances
-	if (Settings->Radius > 0.f)
+	if (SettingsRadius > 0.f)
 	{
 		// Check existing instances. Use the Density radius rather than the minimum radius
-		if (CheckForOverlappingSphere(InWorld, Settings, FSphere(Location, Settings->Radius)))
+		if (CheckForOverlappingSphere(InWorld, Settings, FSphere(Location, SettingsRadius)))
 		{
 			return false;
 		}
 
 		// Check with other potential instances we're about to add.
-		const float RadiusSquared = FMath::Square(Settings->Radius);
-		auto TempInstances = PotentialInstanceHash.GetInstancesOverlappingBox(FBox::BuildAABB(Location, FVector(Settings->Radius)));
+		const float RadiusSquared = FMath::Square(SettingsRadius);
+		auto TempInstances = PotentialInstanceHash.GetInstancesOverlappingBox(FBox::BuildAABB(Location, FVector(SettingsRadius)));
 		for (int32 Idx : TempInstances)
 		{
 			if ((PotentialInstanceLocations[Idx] - Location).SizeSquared() < RadiusSquared)
@@ -1371,6 +1386,7 @@ void FEdModeFoliage::CalculatePotentialInstances(const UWorld* InWorld, const UF
 		Bucket.Reserve(DesiredInstances.Num());
 	}
 
+	const bool bSingleIntanceMode = UISettings ? UISettings->IsInAnySingleInstantiationMode() : false;
 	for (const FDesiredFoliageInstance& DesiredInst : DesiredInstances)
 	{
 		FFoliageTraceFilterFunc TraceFilterFunc;
@@ -1405,7 +1421,7 @@ void FEdModeFoliage::CalculatePotentialInstances(const UWorld* InWorld, const UF
 				continue;
 			}
 
-			const bool bValidInstance = CheckLocationForPotentialInstance(InWorld, Settings, Hit.ImpactPoint, Hit.ImpactNormal, PotentialInstanceLocations, PotentialInstanceHash)
+			const bool bValidInstance = CheckLocationForPotentialInstance(InWorld, Settings, bSingleIntanceMode, Hit.ImpactPoint, Hit.ImpactNormal, PotentialInstanceLocations, PotentialInstanceHash)
 				&& VertexMaskCheck(Hit, Settings)
 				&& LandscapeLayerCheck(Hit, Settings, LocalCache, HitWeight);
 			if (bValidInstance)
@@ -2056,7 +2072,7 @@ void FEdModeFoliage::AdjustUnpaintDensity(float Multiplier)
 	UISettings.SetUnpaintDensity(FMath::Clamp(CurrentDensity + AdjustmentAmount * Multiplier, 0.0f, 1.0f));
 }
 
-void FEdModeFoliage::ReapplyInstancesForBrush(UWorld* InWorld, const UFoliageType* Settings, const FSphere& BrushSphere, float Pressure)
+void FEdModeFoliage::ReapplyInstancesForBrush(UWorld* InWorld, const UFoliageType* Settings, const FSphere& BrushSphere, float Pressure, bool bSingleInstanceMode)
 {
 	// Adjust instance density first
 	ReapplyInstancesDensityForBrush(InWorld, Settings, BrushSphere, Pressure);
@@ -2066,12 +2082,12 @@ void FEdModeFoliage::ReapplyInstancesForBrush(UWorld* InWorld, const UFoliageTyp
 		FFoliageInfo* FoliageInfo = (*It);
 		AInstancedFoliageActor* IFA = It.GetActor();
 
-		ReapplyInstancesForBrush(InWorld, IFA, Settings, FoliageInfo, BrushSphere, Pressure);
+		ReapplyInstancesForBrush(InWorld, IFA, Settings, FoliageInfo, BrushSphere, Pressure, bSingleInstanceMode);
 	}
 }
 
 /** Reapply instance settings to exiting instances */
-void FEdModeFoliage::ReapplyInstancesForBrush(UWorld* InWorld, AInstancedFoliageActor* IFA, const UFoliageType* Settings, FFoliageInfo* FoliageInfo, const FSphere& BrushSphere, float Pressure)
+void FEdModeFoliage::ReapplyInstancesForBrush(UWorld* InWorld, AInstancedFoliageActor* IFA, const UFoliageType* Settings, FFoliageInfo* FoliageInfo, const FSphere& BrushSphere, float Pressure, bool bSingleInstanceMode)
 {
 	TArray<int32> ExistingInstances;
 	FoliageInfo->GetInstancesInsideSphere(BrushSphere, ExistingInstances);
@@ -2328,10 +2344,11 @@ void FEdModeFoliage::ReapplyInstancesForBrush(UWorld* InWorld, AInstancedFoliage
 				FoliageInfo->InstanceHash->InsertInstance(Instance.Location, InstanceIndex);
 			}
 
+			const float SettingsRadius = Settings->GetRadius(bSingleInstanceMode);
 			// Cull overlapping based on radius
-			if (Settings->ReapplyRadius && Settings->Radius > 0.f)
+			if (Settings->ReapplyRadius && SettingsRadius > 0.f)
 			{
-				if (FoliageInfo->CheckForOverlappingInstanceExcluding(InstanceIndex, Settings->Radius, InstancesToDelete))
+				if (FoliageInfo->CheckForOverlappingInstanceExcluding(InstanceIndex, SettingsRadius, InstancesToDelete))
 				{
 					InstancesToDelete.Add(InstanceIndex);
 					continue;
@@ -2483,7 +2500,7 @@ void FEdModeFoliage::ApplyBrush(FEditorViewportClient* ViewportClient)
 		else if (UISettings.GetReapplyToolSelected())
 		{
 			// Reapply any settings checked by the user
-			ReapplyInstancesForBrush(World, Settings, BrushSphere, Pressure);
+			ReapplyInstancesForBrush(World, Settings, BrushSphere, Pressure, UISettings.IsInAnySingleInstantiationMode());
 		}
 		else if (UISettings.GetPaintToolSelected())
 		{
@@ -2694,6 +2711,7 @@ void FEdModeFoliage::ApplyPaintBucket_Add(AActor* Actor)
 		}
 	}
 
+	const bool bSingleIntanceMode = UISettings.IsInAnySingleInstantiationMode();
 	for (const auto& MeshUIInfo : FoliageMeshList)
 	{
 		UFoliageType* Settings = MeshUIInfo->Settings;
@@ -2736,7 +2754,7 @@ void FEdModeFoliage::ApplyPaintBucket_Add(AActor* Actor)
 
 					// Check color mask and filters at this location
 					if (!CheckVertexColor(Settings, VertexColor) ||
-						!CheckLocationForPotentialInstance(World, Settings, InstLocation, Triangle.WorldNormal, PotentialInstanceLocations, PotentialInstanceHash))
+						!CheckLocationForPotentialInstance(World, Settings, bSingleIntanceMode, InstLocation, Triangle.WorldNormal, PotentialInstanceLocations, PotentialInstanceHash))
 					{
 						continue;
 					}
@@ -3779,6 +3797,24 @@ EAxisList::Type FEdModeFoliage::GetWidgetAxisToDraw(FWidget::EWidgetMode InWidge
 	}
 }
 
+float FEdModeFoliage::GetPaintingBrushRadius() const
+{
+	float Radius = UISettings.GetRadius();
+	const bool bSingleInstanceMode = UISettings.IsInAnySingleInstantiationMode();
+	if (bSingleInstanceMode)
+	{
+		for (auto& FoliageMeshUI : FoliageMeshList)
+		{
+			UFoliageType* Settings = FoliageMeshUI->Settings;
+			if (Settings->IsSelected)
+			{
+				Radius = FMath::Max(Radius, Settings->GetRadius(bSingleInstanceMode));
+			}
+		}
+	}
+	return Radius;
+}
+
 /** Load UI settings from ini file */
 void FFoliageUISettings::Load()
 {
diff --git a/Engine/Source/Editor/FoliageEdit/Private/FoliageEdMode.h b/Engine/Source/Editor/FoliageEdit/Private/FoliageEdMode.h
--- a/Engine/Source/Editor/FoliageEdit/Private/FoliageEdMode.h
+++ b/Engine/Source/Editor/FoliageEdit/Private/FoliageEdMode.h
@@ -559,8 +559,8 @@ private:
 
 	/** Reapply instance settings to exiting instances */
 	void ReapplyInstancesDensityForBrush(UWorld* InWorld, const UFoliageType* Settings, const FSphere& BrushSphere, float Pressure);
-	void ReapplyInstancesForBrush(UWorld* InWorld, const UFoliageType* Settings, const FSphere& BrushSphere, float Pressure);
-	void ReapplyInstancesForBrush(UWorld* InWorld, AInstancedFoliageActor* IFA, const UFoliageType* Settings, FFoliageInfo* MeshInfo, const FSphere& BrushSphere, float Pressure);
+	void ReapplyInstancesForBrush(UWorld* InWorld, const UFoliageType* Settings, const FSphere& BrushSphere, float Pressure, bool bSingleInstanceMode);
+	void ReapplyInstancesForBrush(UWorld* InWorld, AInstancedFoliageActor* IFA, const UFoliageType* Settings, FFoliageInfo* MeshInfo, const FSphere& BrushSphere, float Pressure, bool bSingleInstanceMode);
 
 	/** Select instances inside the brush. */
 	void SelectInstancesForBrush(UWorld* InWorld, const UFoliageType* Settings, const FSphere& BrushSphere, bool bSelect);
@@ -620,6 +620,8 @@ private:
 	/** Set the brush mesh opacity */
 	void SetBrushOpacity(const float InOpacity);
 
+	float GetPaintingBrushRadius() const;
+
 	/** Called if the foliage tree is outdated */
 	void RebuildFoliageTree(const UFoliageType* Settings);
 
@@ -631,6 +633,8 @@ private:
 
 	/** The dynamic material of the sphere brush. */
 	class UMaterialInstanceDynamic* BrushMID;
+	FColor BrushDefaultHighlightColor;
+	FColor BrushCurrentHighlightColor;
 
 	/** Default opacity received from the brush material to reset it when closing. */
 	float DefaultBrushOpacity;
diff --git a/Engine/Source/Runtime/Foliage/Public/FoliageType.h b/Engine/Source/Runtime/Foliage/Public/FoliageType.h
--- a/Engine/Source/Runtime/Foliage/Public/FoliageType.h
+++ b/Engine/Source/Runtime/Foliage/Public/FoliageType.h
@@ -133,6 +133,14 @@ public:
 	UPROPERTY(EditAnywhere, Category=Painting, meta=(UIMin=0, ClampMin=0, UIMax = 100000, ClampMax = 100000, ReapplyCondition="ReapplyRadius"))
 	float Radius;
 
+	/** Option to override radius used to detect collision with other instances when painting in single instance mode */
+	UPROPERTY(EditAnywhere, Category = Painting)
+	bool bSingleInstanceModeOverrideRadius = false;
+
+	/** The radius used in single instance mode to detect collision with other instances */
+	UPROPERTY(EditAnywhere, Category = Painting, meta = (EditCondition = "bSingleInstanceModeOverrideRadius", UIMin = 0, ClampMin = 0, UIMax = 100000, ClampMax = 100000))
+	float SingleInstanceModeRadius = 0.f;
+
 	/** Specifies foliage instance scaling behavior when painting. */
 	UPROPERTY(EditAnywhere, Category=Painting, meta=(ReapplyCondition="ReapplyScaling"))
 	EFoliageScaling Scaling;
@@ -338,6 +346,8 @@ public:
 	uint32 IsSelected:1;
 #endif
 public:
+	FOLIAGE_API float GetRadius(bool bInSingleInstanceMode) const { return bInSingleInstanceMode && bSingleInstanceModeOverrideRadius ? SingleInstanceModeRadius : Radius; }
+
 	// PROCEDURAL
 
 	FOLIAGE_API float GetSeedDensitySquared() const { return InitialSeedDensity * InitialSeedDensity; }
