Foliage: Optimize Selection with Foliage Actors

diff --git a/Engine/Source/Editor/FoliageEdit/Private/FoliageEdMode.cpp b/Engine/Source/Editor/FoliageEdit/Private/FoliageEdMode.cpp
--- a/Engine/Source/Editor/FoliageEdit/Private/FoliageEdMode.cpp
+++ b/Engine/Source/Editor/FoliageEdit/Private/FoliageEdMode.cpp
@@ -72,6 +72,24 @@ namespace VREd
 	static FAutoConsoleVariable FoliageOpacity(TEXT("VREd.FoliageOpacity"), 0.02f, TEXT("The foliage brush opacity."));
 }
 
+class FEdModeFoliageSelectionUpdate
+{
+public:
+	FEdModeFoliageSelectionUpdate(FEdModeFoliage* InMode)
+		: Mode(InMode)
+	{
+		Mode->BeginSelectionUpdate();
+	}
+
+	~FEdModeFoliageSelectionUpdate()
+	{
+		Mode->EndSelectionUpdate();
+	}
+
+private:
+	FEdModeFoliage* Mode;
+};
+
 //
 // FFoliageMeshUIInfo
 //
@@ -210,6 +228,8 @@ FEdModeFoliage::FEdModeFoliage()
 	, bAdjustBrushRadius(false)
 	, FoliageMeshListSortMode(EColumnSortMode::Ascending)
 	, FoliageInteractor(nullptr)
+	, UpdateSelectionCounter(0)
+	, bHasDeferredSelectionNotification(false)
 {
 	// Load resources and construct brush component
 	UStaticMesh* StaticMesh = nullptr;
@@ -673,6 +693,7 @@ void FEdModeFoliage::OnVRAction(class FEditorViewportClient& ViewportClient, UVi
 							if (HitResult.GetActor() != nullptr)
 							{
 								// Clear all currently selected instances
+								FEdModeFoliageSelectionUpdate Scope(this);
 								SelectInstances(ViewportClient.GetWorld(), false);
 								for (auto& FoliageMeshUI : FoliageMeshList)
 								{
@@ -754,7 +775,14 @@ void FEdModeFoliage::NotifyActorSelectionChanged(bool bSelect, const TArray<AAct
 		const bool bSelectEvenIfHidden = true;
 		GEditor->SelectActor(Actor, bSelect, bNotify, bSelectEvenIfHidden);
 	}
-	GEditor->NoteSelectionChange();
+	
+	// Defer Notification if we are in a selection update scope
+	bHasDeferredSelectionNotification = UpdateSelectionCounter > 0;
+
+	if (!bHasDeferredSelectionNotification)
+	{
+		GEditor->NoteSelectionChange();
+	}
 }
 
 /** When the user changes the current tool in the UI */
@@ -1512,6 +1540,22 @@ void FEdModeFoliage::RebuildFoliageTree(const UFoliageType* Settings)
 	}
 }
 
+void FEdModeFoliage::BeginSelectionUpdate()
+{
+	UpdateSelectionCounter++;
+}
+
+void FEdModeFoliage::EndSelectionUpdate()
+{
+	check(UpdateSelectionCounter > 0);
+	UpdateSelectionCounter--;
+	if (UpdateSelectionCounter == 0 && bHasDeferredSelectionNotification)
+	{
+		GEditor->NoteSelectionChange();
+		bHasDeferredSelectionNotification = false;
+	}
+}
+
 bool FEdModeFoliage::AddInstancesImp(UWorld* InWorld, const UFoliageType* Settings, const TArray<FDesiredFoliageInstance>& DesiredInstances, const TArray<int32>& ExistingInstanceBuckets, const float Pressure, LandscapeLayerCacheData* LandscapeLayerCachesPtr, const FFoliageUISettings* UISettings, const FFoliagePaintingGeometryFilter* OverrideGeometryFilter, bool InRebuildFoliageTree)
 {
 	SCOPE_CYCLE_COUNTER(STAT_FoliageAddInstanceImp);
@@ -1730,6 +1774,7 @@ void FEdModeFoliage::RemoveInstancesForBrush(UWorld* InWorld, const UFoliageType
 
 void FEdModeFoliage::SelectInstanceAtLocation(UWorld* InWorld, const UFoliageType* Settings, const FVector& Location, bool bSelect)
 {
+	FEdModeFoliageSelectionUpdate Scope(this);
 	for (FFoliageInfoIterator It(InWorld, Settings); It; ++It)
 	{
 		FFoliageInfo* FoliageInfo = (*It);
@@ -1749,6 +1794,7 @@ void FEdModeFoliage::SelectInstanceAtLocation(UWorld* InWorld, const UFoliageTyp
 
 void FEdModeFoliage::SelectInstancesForBrush(UWorld* InWorld, const UFoliageType* Settings, const FSphere& BrushSphere, bool bSelect)
 {
+	FEdModeFoliageSelectionUpdate Scope(this);
 	for (FFoliageInfoIterator It(InWorld, Settings); It; ++It)
 	{
 		FFoliageInfo* FoliageInfo = (*It);
@@ -1765,6 +1811,15 @@ void FEdModeFoliage::SelectInstancesForBrush(UWorld* InWorld, const UFoliageType
 	}
 }
 
+void FEdModeFoliage::SelectInstances(const TArray<const UFoliageType *>& FoliageTypes, bool bSelect)
+{
+	FEdModeFoliageSelectionUpdate Scope(this);
+	for (const UFoliageType* FoliageType : FoliageTypes)
+	{
+		SelectInstances(FoliageType, bSelect);
+	}
+}
+
 void FEdModeFoliage::SelectInstances(const UFoliageType* Settings, bool bSelect)
 {
 	SelectInstances(GetWorld(), Settings, bSelect);
@@ -1772,6 +1827,7 @@ void FEdModeFoliage::SelectInstances(const UFoliageType* Settings, bool bSelect)
 
 void FEdModeFoliage::SelectInstances(UWorld* InWorld, bool bSelect)
 {
+	FEdModeFoliageSelectionUpdate Scope(this);
 	for (auto& FoliageMeshUI : FoliageMeshList)
 	{
 		UFoliageType* Settings = FoliageMeshUI->Settings;
@@ -1787,6 +1843,7 @@ void FEdModeFoliage::SelectInstances(UWorld* InWorld, bool bSelect)
 
 void FEdModeFoliage::SelectInstances(UWorld* InWorld, const UFoliageType* Settings, bool bSelect)
 {
+	FEdModeFoliageSelectionUpdate Scope(this);
 	for (FFoliageInfoIterator It(InWorld, Settings); It; ++It)
 	{
 		FFoliageInfo* FoliageInfo = (*It);
@@ -1800,6 +1857,7 @@ void FEdModeFoliage::ApplySelection(UWorld* InWorld, bool bApply)
 {
 	GEditor->SelectNone(true, true);
 	
+	FEdModeFoliageSelectionUpdate Scope(this);
 	const int32 NumLevels = InWorld->GetNumLevels();
 	for (int32 LevelIdx = 0; LevelIdx < NumLevels; ++LevelIdx)
 	{
@@ -1953,8 +2011,18 @@ TAutoConsoleVariable<float> CVarOffGroundTreshold(
 	5.0f,
 	TEXT("Maximum distance from base component (in local space) at which instance is still considered as valid"));
 
+void FEdModeFoliage::SelectInvalidInstances(const TArray<const UFoliageType *>& FoliageTypes)
+{
+	FEdModeFoliageSelectionUpdate Scope(this);
+	for (const UFoliageType* FoliageType : FoliageTypes)
+	{
+		SelectInvalidInstances(FoliageType);
+	}
+}
+
 void FEdModeFoliage::SelectInvalidInstances(const UFoliageType* Settings)
 {
+	FEdModeFoliageSelectionUpdate Scope(this);
 	UWorld* InWorld = GetWorld();
 
 	FCollisionQueryParams QueryParams(SCENE_QUERY_STAT(FoliageGroundCheck), true);
@@ -3233,6 +3301,7 @@ void FEdModeFoliage::MoveSelectedFoliageToLevel(ULevel* InTargetLevel)
 
 	const FScopedTransaction Transaction(NSLOCTEXT("UnrealEd", "MoveSelectedFoliageToSelectedLevel", "Move Selected Foliage to Level"));
 	
+	FEdModeFoliageSelectionUpdate Scope(this);
 	// Iterate over all foliage actors in the world and move selected instances to a foliage actor in the target level
 	const int32 NumLevels = World->GetNumLevels();
 	for (int32 LevelIdx = 0; LevelIdx < NumLevels; ++LevelIdx)
diff --git a/Engine/Source/Editor/FoliageEdit/Private/FoliageEdMode.h b/Engine/Source/Editor/FoliageEdit/Private/FoliageEdMode.h
--- a/Engine/Source/Editor/FoliageEdit/Private/FoliageEdMode.h
+++ b/Engine/Source/Editor/FoliageEdit/Private/FoliageEdMode.h
@@ -477,9 +477,15 @@ public:
 	/** Save the foliage type object. If it isn't an asset, will prompt the user for a location to save the new asset. */
 	UFoliageType* SaveFoliageTypeObject(UFoliageType* Settings);
 
+	/** Set/Clear selection for foliage instances of a specific types */
+	void SelectInstances(const TArray<const UFoliageType*>& FoliageTypes, bool bSelect);
+
 	/** Set/Clear selection for foliage instances of specific type  */
 	void SelectInstances(const UFoliageType* Settings, bool bSelect);
 
+	/** Find and select instances that don't have valid base or 'off-ground' */
+	void SelectInvalidInstances(const TArray<const UFoliageType*>& FoliageTypes);
+
 	/** Find and select instances that don't have valid base or 'off-ground' */
 	void SelectInvalidInstances(const UFoliageType* Settings);
 
@@ -625,6 +631,12 @@ private:
 	/** Called if the foliage tree is outdated */
 	void RebuildFoliageTree(const UFoliageType* Settings);
 
+	/** Increments a counter that prevents sending out notifications until end selection is called  and counter reaches 0 */
+	void BeginSelectionUpdate();
+
+	/** Decrements counter and will notify editor of selection change if counter reaches 0 */
+	void EndSelectionUpdate();
+
 	bool bBrushTraceValid;
 	FVector BrushLocation;
 	FVector BrushNormal;
@@ -657,5 +669,8 @@ private:
 	/** When painting in VR, this is the hand index that we're painting with.  Otherwise INDEX_NONE. */
 	class UViewportInteractor* FoliageInteractor;
 
+	int32 UpdateSelectionCounter;
+	bool bHasDeferredSelectionNotification;
+	friend class FEdModeFoliageSelectionUpdate;
 };
 
diff --git a/Engine/Source/Editor/FoliageEdit/Private/SFoliageEdit.cpp b/Engine/Source/Editor/FoliageEdit/Private/SFoliageEdit.cpp
--- a/Engine/Source/Editor/FoliageEdit/Private/SFoliageEdit.cpp
+++ b/Engine/Source/Editor/FoliageEdit/Private/SFoliageEdit.cpp
@@ -707,45 +707,50 @@ EVisibility SFoliageEdit::GetVisibility_SelectionOptions() const
 	return IsSelectTool() || IsLassoSelectTool() ? EVisibility::SelfHitTestInvisible : EVisibility::Collapsed;
 }
 
-FReply SFoliageEdit::OnSelectAllInstances()
+void SFoliageEdit::ExecuteOnAllCurrentLevelFoliageTypes(TFunctionRef<void(const TArray<const UFoliageType*>&)> ExecuteFunc)
 {
-	for (FFoliageMeshUIInfoPtr& TypeInfo : FoliageEditMode->GetFoliageMeshList())
+	TArray<FFoliageMeshUIInfoPtr>& FoliageUIList = FoliageEditMode->GetFoliageMeshList();
+	TArray<const UFoliageType*> FoliageTypes;
+	FoliageTypes.Reserve(FoliageUIList.Num());
+
+	for (FFoliageMeshUIInfoPtr& TypeInfo : FoliageUIList)
 	{
 		if (TypeInfo->InstanceCountCurrentLevel > 0)
 		{
-			UFoliageType* FoliageType = TypeInfo->Settings;
-			FoliageEditMode->SelectInstances(FoliageType, true);
+			FoliageTypes.Add(TypeInfo->Settings);
 		}
 	}
 
+	ExecuteFunc(FoliageTypes);
+}
+
+FReply SFoliageEdit::OnSelectAllInstances()
+{
+	ExecuteOnAllCurrentLevelFoliageTypes([&](const TArray<const UFoliageType*>& FoliageTypes)
+	{
+		FoliageEditMode->SelectInstances(FoliageTypes, true);
+	});
+
 	return FReply::Handled();
 }
 
 FReply SFoliageEdit::OnSelectInvalidInstances()
 {
-	for (FFoliageMeshUIInfoPtr& TypeInfo : FoliageEditMode->GetFoliageMeshList())
+	ExecuteOnAllCurrentLevelFoliageTypes([&](const TArray<const UFoliageType*>& FoliageTypes)
 	{
-		if (TypeInfo->InstanceCountCurrentLevel > 0)
-		{
-			const UFoliageType* FoliageType = TypeInfo->Settings;
-			FoliageEditMode->SelectInstances(FoliageType, false);
-			FoliageEditMode->SelectInvalidInstances(FoliageType);
-		}
-	}
+		FoliageEditMode->SelectInstances(FoliageTypes, false);
+		FoliageEditMode->SelectInvalidInstances(FoliageTypes);
+	});
 
 	return FReply::Handled();
 }
 
 FReply SFoliageEdit::OnDeselectAllInstances()
 {
-	for (FFoliageMeshUIInfoPtr& TypeInfo : FoliageEditMode->GetFoliageMeshList())
+	ExecuteOnAllCurrentLevelFoliageTypes([&](const TArray<const UFoliageType*>& FoliageTypes)
 	{
-		if (TypeInfo->InstanceCountCurrentLevel > 0)
-		{
-			UFoliageType* FoliageType = TypeInfo->Settings;
-			FoliageEditMode->SelectInstances(FoliageType, false);
-		}
-	}
+		FoliageEditMode->SelectInstances(FoliageTypes, false);
+	});
 
 	return FReply::Handled();
 }
diff --git a/Engine/Source/Editor/FoliageEdit/Private/SFoliageEdit.h b/Engine/Source/Editor/FoliageEdit/Private/SFoliageEdit.h
--- a/Engine/Source/Editor/FoliageEdit/Private/SFoliageEdit.h
+++ b/Engine/Source/Editor/FoliageEdit/Private/SFoliageEdit.h
@@ -224,7 +224,8 @@ private:	// SELECTION
 	/** Handler to trigger a refresh of the details view when the active tool changes */
 	void HandleOnToolChanged();
 
-	
+	void ExecuteOnAllCurrentLevelFoliageTypes(TFunctionRef<void(const TArray<const UFoliageType*>&)> ExecuteFunc);
+
 private:
 	/** Palette of available foliage types */
 	TSharedPtr<class SFoliagePalette> FoliagePalette;
diff --git a/Engine/Source/Editor/FoliageEdit/Private/SFoliagePalette.cpp b/Engine/Source/Editor/FoliageEdit/Private/SFoliagePalette.cpp
--- a/Engine/Source/Editor/FoliageEdit/Private/SFoliagePalette.cpp
+++ b/Engine/Source/Editor/FoliageEdit/Private/SFoliagePalette.cpp
@@ -1139,31 +1139,40 @@ void SFoliagePalette::OnShowFoliageTypeInCB()
 	}
 }
 
-void SFoliagePalette::OnSelectAllInstances()
+void SFoliagePalette::ExecuteOnSelectedItemFoliageTypes(TFunctionRef<void(const TArray<const UFoliageType*>&)> ExecuteFunc)
 {
-	for (FFoliagePaletteItemModelPtr& PaletteItem : GetActiveViewWidget()->GetSelectedItems())
+	TArray<FFoliagePaletteItemModelPtr> SelectedItems = GetActiveViewWidget()->GetSelectedItems();
+	TArray<const UFoliageType*> FoliageTypes;
+	FoliageTypes.Reserve(SelectedItems.Num());
+	for (FFoliagePaletteItemModelPtr& PaletteItem : SelectedItems)
 	{
-		UFoliageType* FoliageType = PaletteItem->GetFoliageType();
-		FoliageEditMode->SelectInstances(FoliageType, true);
+		FoliageTypes.Add(PaletteItem->GetFoliageType());
 	}
+	ExecuteFunc(FoliageTypes);
+}
+
+void SFoliagePalette::OnSelectAllInstances()
+{
+	ExecuteOnSelectedItemFoliageTypes([&](const TArray<const UFoliageType*>& FoliageTypes)
+	{
+		FoliageEditMode->SelectInstances(FoliageTypes, true);
+	});	
 }
 
 void SFoliagePalette::OnDeselectAllInstances()
 {
-	for (FFoliagePaletteItemModelPtr& PaletteItem : GetActiveViewWidget()->GetSelectedItems())
+	ExecuteOnSelectedItemFoliageTypes([&](const TArray<const UFoliageType*>& FoliageTypes)
 	{
-		UFoliageType* FoliageType = PaletteItem->GetFoliageType();
-		FoliageEditMode->SelectInstances(FoliageType, false);
-	}
+		FoliageEditMode->SelectInstances(FoliageTypes, false);
+	});
 }
 
 void SFoliagePalette::OnSelectInvalidInstances()
 {
-	for (FFoliagePaletteItemModelPtr& PaletteItem : GetActiveViewWidget()->GetSelectedItems())
+	ExecuteOnSelectedItemFoliageTypes([&](const TArray<const UFoliageType*>& FoliageTypes)
 	{
-		UFoliageType* FoliageType = PaletteItem->GetFoliageType();
-		FoliageEditMode->SelectInvalidInstances(FoliageType);
-	}
+		FoliageEditMode->SelectInvalidInstances(FoliageTypes);
+	});
 }
 
 bool SFoliagePalette::CanSelectInstances() const
diff --git a/Engine/Source/Editor/FoliageEdit/Private/SFoliagePalette.h b/Engine/Source/Editor/FoliageEdit/Private/SFoliagePalette.h
--- a/Engine/Source/Editor/FoliageEdit/Private/SFoliagePalette.h
+++ b/Engine/Source/Editor/FoliageEdit/Private/SFoliagePalette.h
@@ -177,6 +177,9 @@ private:	// CONTEXT MENU
 	/** Handler for 'Select Invalid Instances' command  */
 	void OnSelectInvalidInstances();
 
+	/** Executes Function on gathered list of foliage types from the selected palette items */
+	void ExecuteOnSelectedItemFoliageTypes(TFunctionRef<void(const TArray<const UFoliageType*>&)> ExecuteFunc);
+
 	/** @return Whether selecting instances is currently possible */
 	bool CanSelectInstances() const;
 
diff --git a/Engine/Source/Runtime/Foliage/Private/FoliageActor.cpp b/Engine/Source/Runtime/Foliage/Private/FoliageActor.cpp
--- a/Engine/Source/Runtime/Foliage/Private/FoliageActor.cpp
+++ b/Engine/Source/Runtime/Foliage/Private/FoliageActor.cpp
@@ -250,19 +250,21 @@ void FFoliageActor::Reapply(AInstancedFoliageActor* IFA, const UFoliageType* Fol
 	}
 }
 
-void FFoliageActor::SelectInstances(bool bSelect, int32 InstanceIndex, int32 Count)
+void FFoliageActor::SelectAllInstances(bool bSelect)
 {
-	TArray<AActor*> Selection;
-	Selection.Reserve(Count);
-	check(ActorInstances.Num() >= InstanceIndex + Count);
-	for (int32 i = InstanceIndex; i < (InstanceIndex + Count); ++i)
-	{
-		if (ActorInstances[i])
-		{
-			Selection.Add(ActorInstances[i]);
-		}
-	}
-	AInstancedFoliageActor::SelectionChanged.Broadcast(bSelect, Selection);
+	AInstancedFoliageActor::SelectionChanged.Broadcast(bSelect, ActorInstances);
+}
+
+void FFoliageActor::SelectInstance(bool bSelect, int32 Index)
+{
+	TArray<AActor*> SingleInstance;
+	SingleInstance.Add(ActorInstances[Index]);
+	AInstancedFoliageActor::SelectionChanged.Broadcast(bSelect, SingleInstance);
+}
+
+void FFoliageActor::SelectInstances(bool bSelect, const TSet<int32>& SelectedIndices)
+{
+	AInstancedFoliageActor::SelectionChanged.Broadcast(bSelect, GetActorsFromSelectedIndices(SelectedIndices));
 }
 
 void FFoliageActor::ApplySelection(bool bApply, const TSet<int32>& SelectedIndices)
diff --git a/Engine/Source/Runtime/Foliage/Private/FoliageActor.h b/Engine/Source/Runtime/Foliage/Private/FoliageActor.h
--- a/Engine/Source/Runtime/Foliage/Private/FoliageActor.h
+++ b/Engine/Source/Runtime/Foliage/Private/FoliageActor.h
@@ -37,7 +37,9 @@ struct FFoliageActor : public FFoliageImpl
 	int32 FindIndex(const AActor* InActor) const;
 	virtual int32 FindIndex(const UPrimitiveComponent* HitComponent) const override;
 
-	virtual void SelectInstances(bool bSelect, int32 InstanceIndex, int32 Count) override;
+	virtual void SelectAllInstances(bool bSelect) override;
+	virtual void SelectInstance(bool bSelect, int32 Index) override;
+	virtual void SelectInstances(bool bSelect, const TSet<int32>& SelectedIndices) override;
 	virtual void ApplySelection(bool bApply, const TSet<int32>& SelectedIndices) override;
 	virtual void ClearSelection(const TSet<int32>& SelectedIndices) override;
 	virtual void Refresh(AInstancedFoliageActor* IFA, const TArray<FFoliageInstance>& Instances, bool Async, bool Force) override;
diff --git a/Engine/Source/Runtime/Foliage/Private/InstancedFoliage.cpp b/Engine/Source/Runtime/Foliage/Private/InstancedFoliage.cpp
--- a/Engine/Source/Runtime/Foliage/Private/InstancedFoliage.cpp
+++ b/Engine/Source/Runtime/Foliage/Private/InstancedFoliage.cpp
@@ -111,7 +111,9 @@ struct FFoliageStaticMesh : public FFoliageImpl
 	virtual void PostUpdateInstances() override;
 	virtual bool IsOwnedComponent(const UPrimitiveComponent* HitComponent) const override;
 
-	virtual void SelectInstances(bool bSelect, int32 InstanceIndex, int32 Count) override;
+	virtual void SelectAllInstances(bool bSelect) override;
+	virtual void SelectInstance(bool bSelect, int32 Index) override;
+	virtual void SelectInstances(bool bSelect, const TSet<int32>& SelectedIndices) override;
 	virtual void ApplySelection(bool bApply, const TSet<int32>& SelectedIndices) override;
 	virtual void ClearSelection(const TSet<int32>& SelectedIndices) override;
 
@@ -1015,10 +1017,27 @@ bool FFoliageStaticMesh::IsOwnedComponent(const UPrimitiveComponent* HitComponen
 	return Component == HitComponent;
 }
 
-void FFoliageStaticMesh::SelectInstances(bool bSelect, int32 InstanceIndex, int32 Count)
+void FFoliageStaticMesh::SelectAllInstances(bool bSelect)
 {
 	check(Component);
-	Component->SelectInstance(bSelect, InstanceIndex, Count);
+	Component->SelectInstance(bSelect, 0, Component->GetInstanceCount());
+	Component->MarkRenderStateDirty();
+}
+
+void FFoliageStaticMesh::SelectInstance(bool bSelect, int32 Index)
+{
+	check(Component);
+	Component->SelectInstance(bSelect, Index);
+	Component->MarkRenderStateDirty();
+}
+
+void FFoliageStaticMesh::SelectInstances(bool bSelect, const TSet<int32>& SelectedIndices)
+{
+	check(Component);
+	for (int32 i : SelectedIndices)
+	{
+		Component->SelectInstance(bSelect, i);
+	}
 	Component->MarkRenderStateDirty();
 }
 
@@ -1881,6 +1900,8 @@ void FFoliageInfo::PostUpdateInstances(AInstancedFoliageActor* InIFA, const TArr
 {
 	if (InInstancesUpdated.Num())
 	{
+		TSet<int32> UpdateSelectedIndices;
+		UpdateSelectedIndices.Reserve(InInstancesUpdated.Num());
 		for (TArray<int32>::TConstIterator It(InInstancesUpdated); It; ++It)
 		{
 			int32 InstanceIndex = *It;
@@ -1899,10 +1920,15 @@ void FFoliageInfo::PostUpdateInstances(AInstancedFoliageActor* InIFA, const TArr
 			// Reselect the instance to update the render update to include selection as by default it gets removed
 			if (InUpdateSelection)
 			{
-				Implementation->SelectInstances(true, InstanceIndex, 1);
+				UpdateSelectedIndices.Add(InstanceIndex);
 			}
 		}
 
+		if (UpdateSelectedIndices.Num() > 0)
+		{
+			Implementation->SelectInstances(true, UpdateSelectedIndices);
+		}
+
 		Implementation->PostUpdateInstances();
 	}
 }
@@ -2075,7 +2101,7 @@ void FFoliageInfo::SelectInstances(AInstancedFoliageActor* InIFA, bool bSelect)
 				SelectedIndices.Add(i);
 			}
 
-			Implementation->SelectInstances(true, 0, SelectedIndices.Num());
+			Implementation->SelectAllInstances(true);
 		}
 		else
 		{
@@ -2089,6 +2115,8 @@ void FFoliageInfo::SelectInstances(AInstancedFoliageActor* InIFA, bool bSelect,
 {
 	if (InInstances.Num())
 	{
+		TSet<int32> ModifiedSelection;
+		ModifiedSelection.Reserve(InInstances.Num());
 		check(Implementation->IsInitialized());
 		if (bSelect)
 		{
@@ -2099,7 +2127,7 @@ void FFoliageInfo::SelectInstances(AInstancedFoliageActor* InIFA, bool bSelect,
 			for (int32 i : InInstances)
 			{
 				SelectedIndices.Add(i);
-				Implementation->SelectInstances(true, i, 1);
+				ModifiedSelection.Add(i);
 			}
 		}
 		else
@@ -2109,13 +2137,11 @@ void FFoliageInfo::SelectInstances(AInstancedFoliageActor* InIFA, bool bSelect,
 			for (int32 i : InInstances)
 			{
 				SelectedIndices.Remove(i);
-			}
-
-			for (int32 i : InInstances)
-			{
-				Implementation->SelectInstances(false, i, 1);
+				ModifiedSelection.Add(i);
 			}
 		}
+
+		Implementation->SelectInstances(bSelect, ModifiedSelection);
 	}
 }
 
@@ -3043,7 +3069,7 @@ void AInstancedFoliageActor::SelectInstance(AActor* InActor, bool bToggle)
 		{
 			bool bIsSelected = Info->SelectedIndices.Contains(index);
 
-			FoliageActor->SelectInstances(false, index, 1);
+			FoliageActor->SelectInstance(false, index);
 			
 			if (bIsSelected)
 			{
@@ -3053,7 +3079,7 @@ void AInstancedFoliageActor::SelectInstance(AActor* InActor, bool bToggle)
 			if (!bToggle || !bIsSelected)
 			{
 				// Add the selection
-				FoliageActor->SelectInstances(true, index, 1);
+				FoliageActor->SelectInstance(true, index);
 
 				Info->SelectedIndices.Add(index);
 			}
diff --git a/Engine/Source/Runtime/Foliage/Public/InstancedFoliage.h b/Engine/Source/Runtime/Foliage/Public/InstancedFoliage.h
--- a/Engine/Source/Runtime/Foliage/Public/InstancedFoliage.h
+++ b/Engine/Source/Runtime/Foliage/Public/InstancedFoliage.h
@@ -192,7 +192,9 @@ struct FFoliageImpl
 	virtual bool IsOwnedComponent(const UPrimitiveComponent* Component) const = 0;
 	virtual int32 FindIndex(const UPrimitiveComponent* HitComponent) const { return INDEX_NONE; }
 
-	virtual void SelectInstances(bool bSelect, int32 InstanceIndex, int32 Count) = 0;
+	virtual void SelectAllInstances(bool bSelect) = 0;
+	virtual void SelectInstance(bool bSelect, int32 Index) = 0;
+	virtual void SelectInstances(bool bSelect, const TSet<int32>& SelectedIndices) = 0;
 	virtual void ApplySelection(bool bApply, const TSet<int32>& SelectedIndices) = 0;
 	virtual void ClearSelection(const TSet<int32>& SelectedIndices) = 0;
 
