Foliage:  New UI Command (Shortcut needs to be set) that selects foliage types in palette based on current foliage instance selection

diff --git a/Engine/Source/Editor/FoliageEdit/Private/FoliageEdMode.cpp b/Engine/Source/Editor/FoliageEdit/Private/FoliageEdMode.cpp
--- a/Engine/Source/Editor/FoliageEdit/Private/FoliageEdMode.cpp
+++ b/Engine/Source/Editor/FoliageEdit/Private/FoliageEdMode.cpp
@@ -2006,6 +2006,28 @@ void FEdModeFoliage::RemoveSelectedInstances(UWorld* InWorld)
 	GEditor->EndTransaction();
 }
 
+void FEdModeFoliage::GetSelectedInstanceFoliageTypes(TArray<const UFoliageType*>& OutFoliageTypes) const
+{
+	const UWorld* CurrentWorld = GetWorld();
+	const int32 NumLevels = CurrentWorld->GetNumLevels();
+	for (int32 LevelIdx = 0; LevelIdx < NumLevels; ++LevelIdx)
+	{
+		ULevel* Level = CurrentWorld->GetLevel(LevelIdx);
+		AInstancedFoliageActor* IFA = AInstancedFoliageActor::GetInstancedFoliageActorForLevel(Level);
+		if (IFA)
+		{
+			bool bHasSelection = false;
+			for (auto& MeshPair : IFA->FoliageInfos)
+			{
+				if (MeshPair.Value->SelectedIndices.Num())
+				{
+					OutFoliageTypes.AddUnique(MeshPair.Key);
+				}
+			}
+		}
+	}
+}
+
 TAutoConsoleVariable<float> CVarOffGroundTreshold(
 	TEXT("foliage.OffGroundThreshold"),
 	5.0f,
diff --git a/Engine/Source/Editor/FoliageEdit/Private/FoliageEdMode.h b/Engine/Source/Editor/FoliageEdit/Private/FoliageEdMode.h
--- a/Engine/Source/Editor/FoliageEdit/Private/FoliageEdMode.h
+++ b/Engine/Source/Editor/FoliageEdit/Private/FoliageEdMode.h
@@ -489,6 +489,9 @@ public:
 	/** Find and select instances that don't have valid base or 'off-ground' */
 	void SelectInvalidInstances(const UFoliageType* Settings);
 
+	/** Returns selected foliage types based on selected foliage instances */
+	void GetSelectedInstanceFoliageTypes(TArray<const UFoliageType*>& OutFoliageTypes) const;
+
 	/** Adjusts the radius of the foliage brush, using the given multiplier to adjust speed */
 	void AdjustBrushRadius(float Multiplier);
 
diff --git a/Engine/Source/Editor/FoliageEdit/Private/FoliagePaletteCommands.cpp b/Engine/Source/Editor/FoliageEdit/Private/FoliagePaletteCommands.cpp
--- a/Engine/Source/Editor/FoliageEdit/Private/FoliagePaletteCommands.cpp
+++ b/Engine/Source/Editor/FoliageEdit/Private/FoliagePaletteCommands.cpp
@@ -10,6 +10,7 @@ void FFoliagePaletteCommands::RegisterCommands()
 	UI_COMMAND(DeactivateFoliageType, "Deactivate", "Sets the selected foliage types in the palette as inactive (i.e. excluded in brush actions).", EUserInterfaceActionType::Button, FInputChord());
 	UI_COMMAND(RemoveFoliageType, "Remove", "Remove this foliage type from the palette. Removes all associated instances as well.", EUserInterfaceActionType::Button, FInputChord(EKeys::Delete));
 	UI_COMMAND(ShowFoliageTypeInCB, "Show in Content Browser", "Show asset in Content Browser.", EUserInterfaceActionType::Button, FInputChord());
+	UI_COMMAND(ReflectSelectionInPalette, "Reflect Selection in Palette.", "Select foliage type in palette based on selected foliage instances.", EUserInterfaceActionType::None, FInputChord());
 	UI_COMMAND(SelectAllInstances, "Select All Instances", "Select all instances of this foliage type (must be in a selection mode).", EUserInterfaceActionType::Button, FInputChord());
 	UI_COMMAND(DeselectAllInstances, "Deselect All Instances", "Deselect all instances of this foliage type (must be in a selection mode).", EUserInterfaceActionType::Button, FInputChord());
 	UI_COMMAND(SelectInvalidInstances, "Select Invalid Instances", "Select all instances of this foliage type that are off ground (must be in a selection mode).", EUserInterfaceActionType::Button, FInputChord());
diff --git a/Engine/Source/Editor/FoliageEdit/Private/FoliagePaletteCommands.h b/Engine/Source/Editor/FoliageEdit/Private/FoliagePaletteCommands.h
--- a/Engine/Source/Editor/FoliageEdit/Private/FoliagePaletteCommands.h
+++ b/Engine/Source/Editor/FoliageEdit/Private/FoliagePaletteCommands.h
@@ -31,6 +31,7 @@ public:
 	TSharedPtr< FUICommandInfo > DeactivateFoliageType;
 	TSharedPtr< FUICommandInfo > RemoveFoliageType;
 	TSharedPtr< FUICommandInfo > ShowFoliageTypeInCB;
+	TSharedPtr< FUICommandInfo > ReflectSelectionInPalette;
 	TSharedPtr< FUICommandInfo > SelectAllInstances;
 	TSharedPtr< FUICommandInfo > DeselectAllInstances;
 	TSharedPtr< FUICommandInfo > SelectInvalidInstances;
diff --git a/Engine/Source/Editor/FoliageEdit/Private/SFoliagePalette.cpp b/Engine/Source/Editor/FoliageEdit/Private/SFoliagePalette.cpp
--- a/Engine/Source/Editor/FoliageEdit/Private/SFoliagePalette.cpp
+++ b/Engine/Source/Editor/FoliageEdit/Private/SFoliagePalette.cpp
@@ -526,6 +526,10 @@ void SFoliagePalette::BindCommands()
 		Commands.ShowFoliageTypeInCB,
 		FExecuteAction::CreateSP(this, &SFoliagePalette::OnShowFoliageTypeInCB));
 
+	UICommandList->MapAction(
+		Commands.ReflectSelectionInPalette,
+		FExecuteAction::CreateSP(this, &SFoliagePalette::OnReflectSelectionInPalette));
+
 	UICommandList->MapAction(
 		Commands.SelectAllInstances,
 		FExecuteAction::CreateSP(this, &SFoliagePalette::OnSelectAllInstances),
@@ -1139,6 +1143,30 @@ void SFoliagePalette::OnShowFoliageTypeInCB()
 	}
 }
 
+void SFoliagePalette::OnReflectSelectionInPalette()
+{
+	TArray<const UFoliageType*> SelectedFoliageTypes;
+	FoliageEditMode->GetSelectedInstanceFoliageTypes(SelectedFoliageTypes);
+	SelectFoliageTypesInPalette(SelectedFoliageTypes);
+}
+
+void SFoliagePalette::SelectFoliageTypesInPalette(const TArray<const UFoliageType*>& FoliageTypes)
+{
+	TArray<FFoliagePaletteItemModelPtr> SelectedItems;
+	SelectedItems.Reserve(FoliageTypes.Num());
+
+	for (FFoliagePaletteItemModelPtr& PaletteItem : FilteredItems)
+	{
+		if (FoliageTypes.Contains(PaletteItem->GetFoliageType()))
+		{
+			SelectedItems.Add(PaletteItem);
+		}
+	}
+	
+	GetActiveViewWidget()->ClearSelection();
+	GetActiveViewWidget()->SetItemSelection(SelectedItems, true);
+}
+
 void SFoliagePalette::ExecuteOnSelectedItemFoliageTypes(TFunctionRef<void(const TArray<const UFoliageType*>&)> ExecuteFunc)
 {
 	TArray<FFoliagePaletteItemModelPtr> SelectedItems = GetActiveViewWidget()->GetSelectedItems();
diff --git a/Engine/Source/Editor/FoliageEdit/Private/SFoliagePalette.h b/Engine/Source/Editor/FoliageEdit/Private/SFoliagePalette.h
--- a/Engine/Source/Editor/FoliageEdit/Private/SFoliagePalette.h
+++ b/Engine/Source/Editor/FoliageEdit/Private/SFoliagePalette.h
@@ -183,6 +183,12 @@ private:	// CONTEXT MENU
 	/** @return Whether selecting instances is currently possible */
 	bool CanSelectInstances() const;
 
+	/** Handler for 'Reflect Selection in Palette ' command */
+	void OnReflectSelectionInPalette();
+
+	/** Selects Foliage Type in palette */
+	void SelectFoliageTypesInPalette(const TArray<const UFoliageType*>& FoliageTypes);
+
 private:	// THUMBNAIL VIEW
 
 	/** Creates a thumbnail tile for the given foliage type */
