Foliage Actor: Option to Include Actors into Foliage & Exclude Actors from Foliage

#rnx

diff --git a/Engine/Source/Editor/FoliageEdit/FoliageEdit.build.cs b/Engine/Source/Editor/FoliageEdit/FoliageEdit.build.cs
--- a/Engine/Source/Editor/FoliageEdit/FoliageEdit.build.cs
+++ b/Engine/Source/Editor/FoliageEdit/FoliageEdit.build.cs
@@ -19,6 +19,7 @@ public class FoliageEdit : ModuleRules
                 "EditorStyle",
 				"RenderCore",
 				"LevelEditor",
+				"SceneOutliner",
 				"Landscape",
                 "PropertyEditor",
                 "DetailCustomizations",
diff --git a/Engine/Source/Editor/FoliageEdit/Private/FoliageEdMode.cpp b/Engine/Source/Editor/FoliageEdit/Private/FoliageEdMode.cpp
--- a/Engine/Source/Editor/FoliageEdit/Private/FoliageEdMode.cpp
+++ b/Engine/Source/Editor/FoliageEdit/Private/FoliageEdMode.cpp
@@ -6,6 +6,7 @@
 #include "Components/PrimitiveComponent.h"
 #include "Components/StaticMeshComponent.h"
 #include "FoliageType.h"
+#include "FoliageType_Actor.h"
 #include "Components/InstancedStaticMeshComponent.h"
 #include "StaticMeshResources.h"
 #include "FoliageInstancedStaticMeshComponent.h"
@@ -30,6 +31,7 @@
 #include "Toolkits/ToolkitManager.h"
 #include "FoliageEditActions.h"
 #include "FoliageHelper.h"
+#include "ISceneOutliner.h"
 
 #include "AssetRegistryModule.h"
 #include "Misc/ScopeExit.h"
@@ -393,6 +395,7 @@ void FEdModeFoliage::Enter()
 	FWorldDelegates::LevelAddedToWorld.AddSP(this, &FEdModeFoliage::NotifyLevelAddedToWorld);
 	FWorldDelegates::LevelRemovedFromWorld.AddSP(this, &FEdModeFoliage::NotifyLevelRemovedFromWorld);
 	AInstancedFoliageActor::SelectionChanged.AddSP(this, &FEdModeFoliage::NotifyActorSelectionChanged);
+	AInstancedFoliageActor::InstanceCountChanged.AddSP(this, &FEdModeFoliage::OnInstanceCountUpdated);
 
 	FAssetRegistryModule& AssetRegistryModule = FModuleManager::GetModuleChecked<FAssetRegistryModule>(TEXT("AssetRegistry"));
 	AssetRegistryModule.Get().OnAssetRemoved().AddSP(this, &FEdModeFoliage::NotifyAssetRemoved);
@@ -496,6 +499,7 @@ void FEdModeFoliage::Exit()
 	FWorldDelegates::LevelAddedToWorld.RemoveAll(this);
 	FWorldDelegates::LevelRemovedFromWorld.RemoveAll(this);
 	AInstancedFoliageActor::SelectionChanged.RemoveAll(this);
+	AInstancedFoliageActor::InstanceCountChanged.RemoveAll(this);
 
 	if (FModuleManager::Get().IsModuleLoaded(TEXT("AssetRegistry")))
 	{
@@ -1811,6 +1815,120 @@ void FEdModeFoliage::SelectInstancesForBrush(UWorld* InWorld, const UFoliageType
 	}
 }
 
+void RefreshSceneOutliner()
+{
+	// SceneOutliner Refresh
+	TWeakPtr<class ILevelEditor> LevelEditor = FModuleManager::GetModuleChecked<FLevelEditorModule>(TEXT("LevelEditor")).GetLevelEditorInstance();
+	if (LevelEditor.IsValid())
+	{
+		TSharedPtr<class ISceneOutliner> SceneOutlinerPtr = LevelEditor.Pin()->GetSceneOutliner();
+		if (SceneOutlinerPtr)
+		{
+			SceneOutlinerPtr->FullRefresh();
+		}
+	}
+}
+
+void FEdModeFoliage::ExcludeFoliageActors(const TArray<const UFoliageType *>& FoliageTypes, bool bOnlyCurrentLevel)
+{
+	FScopedTransaction Transaction(LOCTEXT("ExcludeFoliageActors", "Exclude Actors from Foliage"));
+	TMap<TSubclassOf<AActor>, const UFoliageType_Actor*> ActorFoliageTypes;
+	for (const UFoliageType* FoliageType : FoliageTypes)
+	{
+		if (const UFoliageType_Actor* ActorFoliageType = Cast<UFoliageType_Actor>(FoliageType))
+		{
+			ActorFoliageTypes.Add(ActorFoliageType->ActorClass, ActorFoliageType);
+		}
+	}
+
+	// Go through all sub-levels
+	UWorld* World = GetWorld();
+	const int32 NumLevels = World->GetNumLevels();
+	for (int32 LevelIdx = 0; LevelIdx < NumLevels; ++LevelIdx)
+	{
+		ULevel* Level = World->GetLevel(LevelIdx);
+		if (!bOnlyCurrentLevel || Level == World->GetCurrentLevel())
+		{
+			AInstancedFoliageActor* IFA = AInstancedFoliageActor::GetInstancedFoliageActorForLevel(Level, false);
+			for (auto& Pair : ActorFoliageTypes)
+			{
+				if (TUniqueObj<FFoliageInfo>* FoliageInfoPtr = IFA->FoliageInfos.Find(Pair.Value))
+				{
+					IFA->Modify();
+					(*FoliageInfoPtr)->ExcludeActors();
+					OnInstanceCountUpdated(Pair.Value);
+				}
+			}
+		}
+	}
+
+	RefreshSceneOutliner();
+}
+
+void FEdModeFoliage::IncludeNonFoliageActors(const TArray<const UFoliageType*>& FoliageTypes, bool bOnlyCurrentLevel)
+{
+	FScopedTransaction Transaction(LOCTEXT("IncludeNonFoliageActors", "Include Actors into Foliage"));
+	TMap<const UClass*, const UFoliageType_Actor*> ActorFoliageTypes;
+	for (const UFoliageType* FoliageType : FoliageTypes)
+	{
+		if (const UFoliageType_Actor* ActorFoliageType = Cast<UFoliageType_Actor>(FoliageType))
+		{
+			ActorFoliageTypes.Add(ActorFoliageType->ActorClass.Get(), ActorFoliageType);
+		}
+	}
+
+	TSet<const UFoliageType*> UpdateFoliageTypes;
+
+	// Go through all sub-levels
+	UWorld* World = GetWorld();
+	const int32 NumLevels = World->GetNumLevels();
+	for (int32 LevelIdx = 0; LevelIdx < NumLevels; ++LevelIdx)
+	{
+		ULevel* Level = World->GetLevel(LevelIdx);
+		if (!bOnlyCurrentLevel || Level == World->GetCurrentLevel())
+		{
+			AInstancedFoliageActor* IFA = nullptr;
+
+			for (auto ActorIterator = Level->Actors.CreateIterator(); ActorIterator; ++ActorIterator)
+			{
+				AActor* CurrentActor = *ActorIterator;
+				if (CurrentActor)
+				{
+					if (const UFoliageType_Actor** FoliageTypePtr = ActorFoliageTypes.Find(CurrentActor->GetClass()))
+					{
+						if (const UFoliageType* FoliageType = *FoliageTypePtr)
+						{
+							if (IFA == nullptr)
+							{
+								IFA = AInstancedFoliageActor::GetInstancedFoliageActorForLevel(Level, true);
+							}
+
+							
+
+							FFoliageInfo* FoliageInfo;
+							IFA->Modify();
+							IFA->AddFoliageType(FoliageType, &FoliageInfo);
+							if (FoliageInfo)
+							{
+
+								FoliageInfo->IncludeActor(IFA, FoliageType, CurrentActor);
+								UpdateFoliageTypes.Add(FoliageType);
+							}
+						}
+					}
+				}
+			}
+		}
+	}
+
+	for (const UFoliageType* FoliageType : UpdateFoliageTypes)
+	{
+		OnInstanceCountUpdated(FoliageType);
+	}
+
+	RefreshSceneOutliner();
+}
+
 void FEdModeFoliage::SelectInstances(const TArray<const UFoliageType *>& FoliageTypes, bool bSelect)
 {
 	FEdModeFoliageSelectionUpdate Scope(this);
diff --git a/Engine/Source/Editor/FoliageEdit/Private/FoliageEdMode.h b/Engine/Source/Editor/FoliageEdit/Private/FoliageEdMode.h
--- a/Engine/Source/Editor/FoliageEdit/Private/FoliageEdMode.h
+++ b/Engine/Source/Editor/FoliageEdit/Private/FoliageEdMode.h
@@ -477,6 +477,10 @@ public:
 	/** Save the foliage type object. If it isn't an asset, will prompt the user for a location to save the new asset. */
 	UFoliageType* SaveFoliageTypeObject(UFoliageType* Settings);
 
+	void IncludeNonFoliageActors(const TArray<const UFoliageType*>& FoliageTypes, bool bOnlyCurrentLevel);
+
+	void ExcludeFoliageActors(const TArray<const UFoliageType*>& FoliageTypes, bool bOnlyCurrentLevel);
+
 	/** Set/Clear selection for foliage instances of a specific types */
 	void SelectInstances(const TArray<const UFoliageType*>& FoliageTypes, bool bSelect);
 
diff --git a/Engine/Source/Editor/FoliageEdit/Private/SFoliagePalette.cpp b/Engine/Source/Editor/FoliageEdit/Private/SFoliagePalette.cpp
--- a/Engine/Source/Editor/FoliageEdit/Private/SFoliagePalette.cpp
+++ b/Engine/Source/Editor/FoliageEdit/Private/SFoliagePalette.cpp
@@ -6,6 +6,7 @@
 #include "Widgets/Input/SCheckBox.h"
 #include "FoliagePaletteItem.h"
 #include "FoliageType.h"
+#include "FoliageType_Actor.h"
 #include "FoliageType_InstancedStaticMesh.h"
 #include "Misc/MessageDialog.h"
 #include "Misc/FeedbackContext.h"
@@ -901,8 +902,14 @@ TSharedPtr<SWidget> SFoliagePalette::ConstructFoliageTypeContextMenu()
 	FMenuBuilder MenuBuilder(true, UICommandList);
 
 	auto SelectedItems = GetActiveViewWidget()->GetSelectedItems();
+	
 	if (SelectedItems.Num() > 0)
 	{
+		bool bSelectionContainsActorFoliage = SelectedItems.IndexOfByPredicate([](const FFoliagePaletteItemModelPtr& PaletteItem)
+		{
+			return Cast<const UFoliageType_Actor>(PaletteItem->GetFoliageType()) != nullptr;
+		}) != INDEX_NONE;
+		
 		const bool bShowSaveAsOption = SelectedItems.Num() == 1 && !SelectedItems[0]->IsAsset() && !SelectedItems[0]->IsBlueprint();
 		if (bShowSaveAsOption)
 		{
@@ -950,6 +957,27 @@ TSharedPtr<SWidget> SFoliagePalette::ConstructFoliageTypeContextMenu()
 				LOCTEXT("ReplaceFoliageType_ToolTip", "Replaces selected foliage type with another foliage type asset"),
 				FNewMenuDelegate::CreateSP(this, &SFoliagePalette::FillReplaceFoliageTypeSubmenu));
 
+			if (bSelectionContainsActorFoliage)
+			{
+				MenuBuilder.AddSubMenu(
+					LOCTEXT("IncludeNonFoliageActors", "Include Non Foliage Actors..."),
+					LOCTEXT("IncludeNonFoliageActors_ToolTip", "Include non foliage type actors with matching blueprint into foliage"),
+					FNewMenuDelegate::CreateLambda([&](FMenuBuilder& SubMenuBuilder)
+				{
+					SubMenuBuilder.AddMenuEntry(LOCTEXT("IncludeNonFoliageActors_CurrentLevel", "Current Level"), FText(), FSlateIcon(), FExecuteAction::CreateSP(this, &SFoliagePalette::OnIncludeNonFoliageActors, true));
+					SubMenuBuilder.AddMenuEntry(LOCTEXT("IncludeNonFoliageActors_AllLevels", "All Levels"), FText(), FSlateIcon(), FExecuteAction::CreateSP(this, &SFoliagePalette::OnIncludeNonFoliageActors, false));
+				}));
+
+				MenuBuilder.AddSubMenu(
+					LOCTEXT("ExcludeFoliageActors", "Exclude Actors..."),
+					LOCTEXT("ExcludeFoliageActors_ToolTip", "Exclude actors with foliage type from foliage"),
+					FNewMenuDelegate::CreateLambda([&](FMenuBuilder& SubMenuBuilder)
+				{
+					SubMenuBuilder.AddMenuEntry(LOCTEXT("ExcludeFoliageActors_CurrentLevel", "Current Level"), FText(), FSlateIcon(), FExecuteAction::CreateSP(this, &SFoliagePalette::OnExcludeFoliageActors, true));
+					SubMenuBuilder.AddMenuEntry(LOCTEXT("ExcludeFoliageActors_AllLevels", "All Levels"), FText(), FSlateIcon(), FExecuteAction::CreateSP(this, &SFoliagePalette::OnExcludeFoliageActors, false));
+				}));
+			}
+
 			MenuBuilder.AddMenuEntry(Commands.ShowFoliageTypeInCB);
 		}
 		MenuBuilder.EndSection();
@@ -966,6 +994,22 @@ TSharedPtr<SWidget> SFoliagePalette::ConstructFoliageTypeContextMenu()
 	return MenuBuilder.MakeWidget();
 }
 
+void SFoliagePalette::OnIncludeNonFoliageActors(bool bOnlyCurrentLevel)
+{
+	ExecuteOnSelectedItemFoliageTypes([&, bOnlyCurrentLevel](const TArray<const UFoliageType*>& FoliageTypes)
+	{
+		FoliageEditMode->IncludeNonFoliageActors(FoliageTypes, bOnlyCurrentLevel);
+	});
+}
+
+void SFoliagePalette::OnExcludeFoliageActors(bool bOnlyCurrentLevel)
+{
+	ExecuteOnSelectedItemFoliageTypes([&, bOnlyCurrentLevel](const TArray<const UFoliageType*>& FoliageTypes)
+	{
+		FoliageEditMode->ExcludeFoliageActors(FoliageTypes, bOnlyCurrentLevel);
+	});
+}
+
 void SFoliagePalette::OnSaveSelected()
 {
 	for (FFoliagePaletteItemModelPtr& PaletteItem : GetActiveViewWidget()->GetSelectedItems())
diff --git a/Engine/Source/Editor/FoliageEdit/Private/SFoliagePalette.h b/Engine/Source/Editor/FoliageEdit/Private/SFoliagePalette.h
--- a/Engine/Source/Editor/FoliageEdit/Private/SFoliagePalette.h
+++ b/Engine/Source/Editor/FoliageEdit/Private/SFoliagePalette.h
@@ -136,6 +136,12 @@ private:	// GENERAL
 
 private:	// CONTEXT MENU
 
+	/**  Finds actors with class matching foliage type and merges them into foliage instances */
+	void OnIncludeNonFoliageActors(bool bOnlyCurrentLevel);
+
+	/** Excludes actors from matching foliage type and make them regular actors. */
+	void OnExcludeFoliageActors(bool bOnlyCurrentLevel);
+
 	/** @return the SWidget containing the context menu */
 	TSharedPtr<SWidget> ConstructFoliageTypeContextMenu();
 
diff --git a/Engine/Source/Editor/LevelEditor/Private/LevelEditor.cpp b/Engine/Source/Editor/LevelEditor/Private/LevelEditor.cpp
--- a/Engine/Source/Editor/LevelEditor/Private/LevelEditor.cpp
+++ b/Engine/Source/Editor/LevelEditor/Private/LevelEditor.cpp
@@ -433,7 +433,7 @@ const FLevelViewportCommands& FLevelEditorModule::GetLevelViewportCommands() con
 	return FLevelViewportCommands::Get();
 }
 
-TWeakPtr<class SLevelEditor> FLevelEditorModule::GetLevelEditorInstance() const
+TWeakPtr<class ILevelEditor> FLevelEditorModule::GetLevelEditorInstance() const
 {
 	return LevelEditorInstancePtr;
 }
diff --git a/Engine/Source/Editor/LevelEditor/Private/LevelEditorActions.cpp b/Engine/Source/Editor/LevelEditor/Private/LevelEditorActions.cpp
--- a/Engine/Source/Editor/LevelEditor/Private/LevelEditorActions.cpp
+++ b/Engine/Source/Editor/LevelEditor/Private/LevelEditorActions.cpp
@@ -1550,7 +1550,7 @@ bool FLevelEditorActionCallbacks::Duplicate_CanExecute()
 
 	if (!bCanCopy)
 	{
-		TWeakPtr<class SLevelEditor> LevelEditor = FModuleManager::GetModuleChecked<FLevelEditorModule>(TEXT("LevelEditor")).GetLevelEditorInstance();
+		TWeakPtr<class ILevelEditor> LevelEditor = FModuleManager::GetModuleChecked<FLevelEditorModule>(TEXT("LevelEditor")).GetLevelEditorInstance();
 		if (LevelEditor.IsValid())
 		{
 			TSharedPtr<class ISceneOutliner> SceneOutlinerPtr = LevelEditor.Pin()->GetSceneOutliner();
@@ -1603,7 +1603,7 @@ bool FLevelEditorActionCallbacks::Delete_CanExecute()
 
 	if (!bCanDelete)
 	{
-		TWeakPtr<class SLevelEditor> LevelEditor = FModuleManager::GetModuleChecked<FLevelEditorModule>(TEXT("LevelEditor")).GetLevelEditorInstance();
+		TWeakPtr<class ILevelEditor> LevelEditor = FModuleManager::GetModuleChecked<FLevelEditorModule>(TEXT("LevelEditor")).GetLevelEditorInstance();
 		if (LevelEditor.IsValid())
 		{
 			TSharedPtr<class ISceneOutliner> SceneOutlinerPtr = LevelEditor.Pin()->GetSceneOutliner();
@@ -1630,7 +1630,7 @@ void FLevelEditorActionCallbacks::Rename_Execute()
 	}
 	else
 	{
-		TWeakPtr<class SLevelEditor> LevelEditor = FModuleManager::GetModuleChecked<FLevelEditorModule>(TEXT("LevelEditor")).GetLevelEditorInstance();
+		TWeakPtr<class ILevelEditor> LevelEditor = FModuleManager::GetModuleChecked<FLevelEditorModule>(TEXT("LevelEditor")).GetLevelEditorInstance();
 		if (LevelEditor.IsValid())
 		{
 			TSharedPtr<class ISceneOutliner> SceneOutlinerPtr = LevelEditor.Pin()->GetSceneOutliner();
@@ -1660,7 +1660,7 @@ bool FLevelEditorActionCallbacks::Rename_CanExecute()
 
 	if (!bCanRename)
 	{
-		TWeakPtr<class SLevelEditor> LevelEditor = FModuleManager::GetModuleChecked<FLevelEditorModule>(TEXT("LevelEditor")).GetLevelEditorInstance();
+		TWeakPtr<class ILevelEditor> LevelEditor = FModuleManager::GetModuleChecked<FLevelEditorModule>(TEXT("LevelEditor")).GetLevelEditorInstance();
 		if (LevelEditor.IsValid())
 		{
 			TSharedPtr<class ISceneOutliner> SceneOutlinerPtr = LevelEditor.Pin()->GetSceneOutliner();
@@ -1715,7 +1715,7 @@ bool FLevelEditorActionCallbacks::Cut_CanExecute()
 
 	if (!bCanCut)
 	{
-		TWeakPtr<class SLevelEditor> LevelEditor = FModuleManager::GetModuleChecked<FLevelEditorModule>(TEXT("LevelEditor")).GetLevelEditorInstance();
+		TWeakPtr<class ILevelEditor> LevelEditor = FModuleManager::GetModuleChecked<FLevelEditorModule>(TEXT("LevelEditor")).GetLevelEditorInstance();
 		if (LevelEditor.IsValid())
 		{
 			TSharedPtr<class ISceneOutliner> SceneOutlinerPtr = LevelEditor.Pin()->GetSceneOutliner();
@@ -1768,7 +1768,7 @@ bool FLevelEditorActionCallbacks::Copy_CanExecute()
 
 	if (!bCanCopy)
 	{
-		TWeakPtr<class SLevelEditor> LevelEditor = FModuleManager::GetModuleChecked<FLevelEditorModule>(TEXT("LevelEditor")).GetLevelEditorInstance();
+		TWeakPtr<class ILevelEditor> LevelEditor = FModuleManager::GetModuleChecked<FLevelEditorModule>(TEXT("LevelEditor")).GetLevelEditorInstance();
 		if (LevelEditor.IsValid())
 		{
 			TSharedPtr<class ISceneOutliner> SceneOutlinerPtr = LevelEditor.Pin()->GetSceneOutliner();
@@ -1819,7 +1819,7 @@ bool FLevelEditorActionCallbacks::Paste_CanExecute()
 
 	if (!bCanPaste)
 	{
-		TWeakPtr<class SLevelEditor> LevelEditor = FModuleManager::GetModuleChecked<FLevelEditorModule>(TEXT("LevelEditor")).GetLevelEditorInstance();
+		TWeakPtr<class ILevelEditor> LevelEditor = FModuleManager::GetModuleChecked<FLevelEditorModule>(TEXT("LevelEditor")).GetLevelEditorInstance();
 		if (LevelEditor.IsValid())
 		{
 			TSharedPtr<class ISceneOutliner> SceneOutlinerPtr = LevelEditor.Pin()->GetSceneOutliner();
diff --git a/Engine/Source/Editor/LevelEditor/Private/SLevelEditor.h b/Engine/Source/Editor/LevelEditor/Private/SLevelEditor.h
--- a/Engine/Source/Editor/LevelEditor/Private/SLevelEditor.h
+++ b/Engine/Source/Editor/LevelEditor/Private/SLevelEditor.h
@@ -133,7 +133,7 @@ public:
 	void AttachSequencer( TSharedPtr<SWidget> SequencerWidget, TSharedPtr<IAssetEditorInstance> NewSequencerAssetEditor );
 
 	/** Returns current scene outliner associated with level editor's scene outliner tab, if it exists */
-	TSharedPtr<ISceneOutliner> GetSceneOutliner() { return SceneOutlinerPtr.Pin();  }
+	virtual TSharedPtr<ISceneOutliner> GetSceneOutliner() const override { return SceneOutlinerPtr.Pin();  }
 
 private:
 	
diff --git a/Engine/Source/Editor/LevelEditor/Public/ILevelEditor.h b/Engine/Source/Editor/LevelEditor/Public/ILevelEditor.h
--- a/Engine/Source/Editor/LevelEditor/Public/ILevelEditor.h
+++ b/Engine/Source/Editor/LevelEditor/Public/ILevelEditor.h
@@ -12,6 +12,7 @@
 #include "AssetThumbnail.h"
 #include "Toolkits/IToolkitHost.h"
 
+class ISceneOutliner;
 class ILevelViewport;
 class SLevelViewport;
 
@@ -60,6 +61,7 @@ public:
 	/** Spawns a level editor ToolBox widget (aka. "Modes") */
 	virtual TSharedRef<SWidget> CreateToolBox() = 0;
 
+	virtual TSharedPtr<ISceneOutliner> GetSceneOutliner() const = 0;
 };
 
 
diff --git a/Engine/Source/Editor/LevelEditor/Public/LevelEditor.h b/Engine/Source/Editor/LevelEditor/Public/LevelEditor.h
--- a/Engine/Source/Editor/LevelEditor/Public/LevelEditor.h
+++ b/Engine/Source/Editor/LevelEditor/Public/LevelEditor.h
@@ -145,7 +145,7 @@ public:
 	virtual const class FLevelViewportCommands& GetLevelViewportCommands() const;
 
 	/* @return The pointer to the current level Editor instance */
-	virtual TWeakPtr<class SLevelEditor> GetLevelEditorInstance() const;
+	virtual TWeakPtr<class ILevelEditor> GetLevelEditorInstance() const;
 
 	/* @return The pointer to the level editor tab */
 	virtual TWeakPtr<class SDockTab> GetLevelEditorInstanceTab() const;
diff --git a/Engine/Source/Editor/SceneOutliner/Public/ISceneOutliner.h b/Engine/Source/Editor/SceneOutliner/Public/ISceneOutliner.h
--- a/Engine/Source/Editor/SceneOutliner/Public/ISceneOutliner.h
+++ b/Engine/Source/Editor/SceneOutliner/Public/ISceneOutliner.h
@@ -19,6 +19,9 @@ public:
 	/** Sends a requests to the Scene Outliner to refresh itself the next chance it gets */
 	virtual void Refresh() = 0;
 
+	/** Tells the scene outliner that it should do a full refresh, which will clear the entire tree and rebuild it from scratch. */
+	virtual void FullRefresh() = 0;
+
 	/** @return Returns a string to use for highlighting results in the outliner list */
 	virtual TAttribute<FText> GetFilterHighlightText() const = 0;
 
diff --git a/Engine/Source/Editor/SceneOutliner/Public/SSceneOutliner.h b/Engine/Source/Editor/SceneOutliner/Public/SSceneOutliner.h
--- a/Engine/Source/Editor/SceneOutliner/Public/SSceneOutliner.h
+++ b/Engine/Source/Editor/SceneOutliner/Public/SSceneOutliner.h
@@ -164,6 +164,9 @@ namespace SceneOutliner
 		/** Returns true if clipboard contains folders only */
 		bool CanPasteFoldersOnlyFromClipboard();
 		
+		/** Tells the scene outliner that it should do a full refresh, which will clear the entire tree and rebuild it from scratch. */
+		virtual void FullRefresh() override;
+
 	private:
 		/** Methods that implement structural modification logic for the tree */
 
@@ -176,9 +179,6 @@ namespace SceneOutliner
 		/** Repopulates the entire tree */
 		void RepopulateEntireTree();
 
-		/** Tells the scene outliner that it should do a full refresh, which will clear the entire tree and rebuild it from scratch. */
-		void FullRefresh();
-
 		/** Tells the scene outliner that there was a change in the level actor list. */
 		void OnLevelActorListChanged();
 
diff --git a/Engine/Source/Runtime/Engine/Classes/GameFramework/Actor.h b/Engine/Source/Runtime/Engine/Classes/GameFramework/Actor.h
--- a/Engine/Source/Runtime/Engine/Classes/GameFramework/Actor.h
+++ b/Engine/Source/Runtime/Engine/Classes/GameFramework/Actor.h
@@ -3209,12 +3209,13 @@ private:
 struct FSetActorHiddenInSceneOutliner
 {
 private:
-	FSetActorHiddenInSceneOutliner(AActor* InActor)
+	FSetActorHiddenInSceneOutliner(AActor* InActor, bool bHidden = true)
 	{
-		InActor->bListedInSceneOutliner = false;
+		InActor->bListedInSceneOutliner = !bHidden;
 	}
 
 	friend UWorld;
+	friend class FFoliageHelper;
 };
 #endif
 
diff --git a/Engine/Source/Runtime/Foliage/Private/InstancedFoliage.cpp b/Engine/Source/Runtime/Foliage/Private/InstancedFoliage.cpp
--- a/Engine/Source/Runtime/Foliage/Private/InstancedFoliage.cpp
+++ b/Engine/Source/Runtime/Foliage/Private/InstancedFoliage.cpp
@@ -2195,6 +2195,52 @@ void FFoliageInfo::RemoveBaseComponentOnInstances()
 	}
 }
 
+void FFoliageInfo::IncludeActor(AInstancedFoliageActor* IFA, const UFoliageType* FoliageType, AActor* InActor)
+{
+	if (FFoliageActor* FoliageActor = StaticCast<FFoliageActor*>(Implementation.Get()))
+	{
+		if (FoliageActor->FindIndex(InActor) == INDEX_NONE)
+		{
+			FFoliageInstance NewInstance;
+			NewInstance.BaseComponent = nullptr;
+			NewInstance.BaseId = FFoliageInstanceBaseCache::InvalidBaseId;
+
+			NewInstance.DrawScale3D = InActor->GetActorScale3D();
+			FTransform Transform = InActor->GetTransform();
+			NewInstance.Location = Transform.GetLocation();
+			NewInstance.Rotation = FRotator(Transform.GetRotation());
+			NewInstance.PreAlignRotation = NewInstance.Rotation;
+
+			int32 Index = Instances.Add(NewInstance);
+			InstanceHash->InsertInstance(NewInstance.Location, Index);
+			FoliageActor->PreAddInstances(IFA, FoliageType, 1);
+			FoliageActor->ActorInstances.Add(InActor);
+			InActor->Modify();
+			FFoliageHelper::SetIsOwnedByFoliage(InActor);
+		}
+	}
+}
+
+void FFoliageInfo::ExcludeActors()
+{
+	if (FFoliageActor* FoliageActor = StaticCast<FFoliageActor*>(Implementation.Get()))
+	{
+		SelectedIndices.Empty();
+		Instances.Empty();
+		InstanceHash->Empty();
+		ComponentHash.Empty();
+		for (AActor* Actor : FoliageActor->ActorInstances)
+		{
+			if (Actor != nullptr)
+			{
+				Actor->Modify();
+				FFoliageHelper::SetIsOwnedByFoliage(Actor, false);
+			}
+		}
+		FoliageActor->ActorInstances.Empty();
+	}
+}
+
 TArray<int32> FFoliageInfo::GetInstancesOverlappingBox(const FBox& Box) const
 {
 	return InstanceHash->GetInstancesOverlappingBox(Box);
@@ -3418,6 +3464,7 @@ void AInstancedFoliageActor::Serialize(FArchive& Ar)
 
 #if WITH_EDITOR
 AInstancedFoliageActor::FOnSelectionChanged AInstancedFoliageActor::SelectionChanged;
+AInstancedFoliageActor::FOnInstanceCoundChanged AInstancedFoliageActor::InstanceCountChanged;
 
 void AInstancedFoliageActor::PostInitProperties()
 {
@@ -3915,6 +3962,10 @@ void AInstancedFoliageActor::OnLevelActorDeleted(AActor* InActor)
 						if(InstancesToRemove.Num() > 0)
 						{
 							Info.RemoveInstances(this, InstancesToRemove, false);
+							if (InstanceCountChanged.IsBound())
+							{
+								InstanceCountChanged.Broadcast(FoliageType);
+							}
 						}
 					}
 				}
diff --git a/Engine/Source/Runtime/Foliage/Public/FoliageHelper.h b/Engine/Source/Runtime/Foliage/Public/FoliageHelper.h
--- a/Engine/Source/Runtime/Foliage/Public/FoliageHelper.h
+++ b/Engine/Source/Runtime/Foliage/Public/FoliageHelper.h
@@ -11,7 +11,22 @@ class FFoliageHelper
 {
 public:
 #if WITH_EDITOR
-	static void SetIsOwnedByFoliage(AActor* InActor) { if (InActor) { InActor->Tags.AddUnique("FoliageActorInstance"); } }
-	static bool IsOwnedByFoliage(const AActor* InActor) { return InActor != nullptr && InActor->ActorHasTag("FoliageActorInstance"); }
+	static void SetIsOwnedByFoliage(AActor* InActor, bool bOwned = true) 
+	{ 
+		static FName NAME_FoliageActorInstanceTag(TEXT("FoliageActorInstance"));
+		if (InActor) 
+		{ 
+			FSetActorHiddenInSceneOutliner SetHidden(InActor, bOwned);
+			if (bOwned)
+			{
+				InActor->Tags.AddUnique(NAME_FoliageActorInstanceTag);
+			}
+			else
+			{
+				InActor->Tags.Remove(NAME_FoliageActorInstanceTag);
+			}
+		} 
+	}
+	static bool IsOwnedByFoliage(const AActor* InActor) { return InActor != nullptr && InActor->ActorHasTag(TEXT("FoliageActorInstance")); }
 #endif
 };
\ No newline at end of file
diff --git a/Engine/Source/Runtime/Foliage/Public/InstancedFoliage.h b/Engine/Source/Runtime/Foliage/Public/InstancedFoliage.h
--- a/Engine/Source/Runtime/Foliage/Public/InstancedFoliage.h
+++ b/Engine/Source/Runtime/Foliage/Public/InstancedFoliage.h
@@ -318,6 +318,8 @@ struct FFoliageInfo
 	FOLIAGE_API void ExitEditMode();
 
 	FOLIAGE_API void RemoveBaseComponentOnInstances();
+	FOLIAGE_API void IncludeActor(AInstancedFoliageActor* IFA, const UFoliageType* FoliageType, AActor* InActor);
+	FOLIAGE_API void ExcludeActors();
 #endif
 
 	friend FArchive& operator<<(FArchive& Ar, FFoliageInfo& MeshInfo);
diff --git a/Engine/Source/Runtime/Foliage/Public/InstancedFoliageActor.h b/Engine/Source/Runtime/Foliage/Public/InstancedFoliageActor.h
--- a/Engine/Source/Runtime/Foliage/Public/InstancedFoliageActor.h
+++ b/Engine/Source/Runtime/Foliage/Public/InstancedFoliageActor.h
@@ -104,6 +104,9 @@ public:
 	// Delegate type for selection change events
 	DECLARE_MULTICAST_DELEGATE_TwoParams(FOnSelectionChanged, bool, const TArray<AActor*>&);
 	FOLIAGE_API static FOnSelectionChanged SelectionChanged;
+
+	DECLARE_MULTICAST_DELEGATE_OneParam(FOnInstanceCoundChanged, const UFoliageType*);
+	FOLIAGE_API static FOnInstanceCoundChanged InstanceCountChanged;
 #endif
 	//~ End AActor Interface.
 
