Foliage: Fix exclude crash with null IFA.

#rnx

diff --git a/Engine/Source/Editor/FoliageEdit/Private/FoliageEdMode.cpp b/Engine/Source/Editor/FoliageEdit/Private/FoliageEdMode.cpp
--- a/Engine/Source/Editor/FoliageEdit/Private/FoliageEdMode.cpp
+++ b/Engine/Source/Editor/FoliageEdit/Private/FoliageEdMode.cpp
@@ -1849,14 +1849,16 @@ void FEdModeFoliage::ExcludeFoliageActors(const TArray<const UFoliageType *>& Fo
 		ULevel* Level = World->GetLevel(LevelIdx);
 		if (!bOnlyCurrentLevel || Level == World->GetCurrentLevel())
 		{
-			AInstancedFoliageActor* IFA = AInstancedFoliageActor::GetInstancedFoliageActorForLevel(Level, false);
-			for (auto& Pair : ActorFoliageTypes)
+			if (AInstancedFoliageActor* IFA = AInstancedFoliageActor::GetInstancedFoliageActorForLevel(Level, false))
 			{
-				if (TUniqueObj<FFoliageInfo>* FoliageInfoPtr = IFA->FoliageInfos.Find(Pair.Value))
+				for (auto& Pair : ActorFoliageTypes)
 				{
-					IFA->Modify();
-					(*FoliageInfoPtr)->ExcludeActors();
-					OnInstanceCountUpdated(Pair.Value);
+					if (TUniqueObj<FFoliageInfo>* FoliageInfoPtr = IFA->FoliageInfos.Find(Pair.Value))
+					{
+						IFA->Modify();
+						(*FoliageInfoPtr)->ExcludeActors();
+						OnInstanceCountUpdated(Pair.Value);
+					}
 				}
 			}
 		}
