Added Exclusion Landscape Layers and minimum threshold value in FoliageType to exclude foliage instances in a specific area

diff --git a/Engine/Source/Editor/FoliageEdit/Private/FoliageEdMode.cpp b/Engine/Source/Editor/FoliageEdit/Private/FoliageEdMode.cpp
--- a/Engine/Source/Editor/FoliageEdit/Private/FoliageEdMode.cpp
+++ b/Engine/Source/Editor/FoliageEdit/Private/FoliageEdMode.cpp
@@ -1258,10 +1258,10 @@ static bool CheckVertexColor(const UFoliageType* Settings, const FColor& VertexC
 	return true;
 }
 
-bool LandscapeLayersValid(const UFoliageType* Settings)
+bool IsLandscapeLayersArrayValid(const TArray<FName>& LandscapeLayersArray)
 {
 	bool bValid = false;
-	for (FName LayerName : Settings->LandscapeLayers)
+	for (FName LayerName : LandscapeLayersArray)
 	{
 		bValid |= LayerName != NAME_None;
 	}
@@ -1269,14 +1269,14 @@ bool LandscapeLayersValid(const UFoliageType* Settings)
 	return bValid;
 }
 
-bool GetMaxHitWeight(const FVector& Location, UActorComponent* Component, const UFoliageType* Settings, FEdModeFoliage::LandscapeLayerCacheData* LandscapeLayerCaches, float& OutMaxHitWeight)
+bool GetMaxHitWeight(const FVector& Location, UActorComponent* Component, const TArray<FName>& LandscapeLayersArray, FEdModeFoliage::LandscapeLayerCacheData* LandscapeLayerCaches, float& OutMaxHitWeight)
 {
 	float MaxHitWeight = 0.f;
 	if (ULandscapeHeightfieldCollisionComponent* HitLandscapeCollision = Cast<ULandscapeHeightfieldCollisionComponent>(Component))
 	{
 		if (ULandscapeComponent* HitLandscape = HitLandscapeCollision->RenderComponent.Get())
 		{
-			for (const FName LandscapeLayerName : Settings->LandscapeLayers)
+			for (const FName LandscapeLayerName : LandscapeLayersArray)
 			{
 				// Cache store mapping between component and weight data
 				TMap<ULandscapeComponent*, TArray<uint8> >* LandscapeLayerCache = &LandscapeLayerCaches->FindOrAdd(LandscapeLayerName);;
@@ -1294,10 +1294,19 @@ bool GetMaxHitWeight(const FVector& Location, UActorComponent* Component, const
 	return false;
 }
 
-bool FilterByWeight(float Weight, const UFoliageType* Settings)
+bool IsFilteredByWeight(float Weight, float TestValue, bool bExclusionTest = false)
 {
-	const float WeightNeeded = FMath::Max(Settings->MinimumLayerWeight, FMath::FRand());
-	return Weight < FMath::Max(SMALL_NUMBER, WeightNeeded);
+	if (bExclusionTest)
+	{
+		// Exclusion always tests 
+		const float WeightNeeded = FMath::Max(SMALL_NUMBER, TestValue);
+		return Weight >= WeightNeeded;
+	}
+	else
+	{
+		const float WeightNeeded = FMath::Max(SMALL_NUMBER, FMath::Max(TestValue, FMath::FRand()));
+		return Weight < WeightNeeded;
+	}
 }
 
 bool FEdModeFoliage::IsUsingVertexColorMask(const UFoliageType* Settings)
@@ -1343,10 +1352,21 @@ void FEdModeFoliage::SetBrushOpacity(const float InOpacity)
 bool LandscapeLayerCheck(const FHitResult& Hit, const UFoliageType* Settings, FEdModeFoliage::LandscapeLayerCacheData& LandscapeLayersCache, float& OutHitWeight)
 {
 	OutHitWeight = 1.f;
-	if (LandscapeLayersValid(Settings) && GetMaxHitWeight(Hit.ImpactPoint, Hit.Component.Get(), Settings, &LandscapeLayersCache, OutHitWeight))
+	if (IsLandscapeLayersArrayValid(Settings->LandscapeLayers) && GetMaxHitWeight(Hit.ImpactPoint, Hit.Component.Get(), Settings->LandscapeLayers, &LandscapeLayersCache, OutHitWeight))
+	{
+		// Reject instance randomly in proportion to weight
+		if (IsFilteredByWeight(OutHitWeight, Settings->MinimumLayerWeight))
+		{
+			return false;
+		}
+	}
+
+	float HitWeightExclusion = 1.f;
+	if (IsLandscapeLayersArrayValid(Settings->ExclusionLandscapeLayers) && GetMaxHitWeight(Hit.ImpactPoint, Hit.Component.Get(), Settings->ExclusionLandscapeLayers, &LandscapeLayersCache, HitWeightExclusion))
 	{
 		// Reject instance randomly in proportion to weight
-		if (FilterByWeight(OutHitWeight, Settings))
+		const bool bExclusionTest = true;
+		if (IsFilteredByWeight(HitWeightExclusion, Settings->MinimumExclusionLayerWeight, bExclusionTest))
 		{
 			return false;
 		}
@@ -1663,7 +1683,7 @@ void FEdModeFoliage::AddInstancesForBrush(UWorld* InWorld, const UFoliageType* S
 	SCOPE_CYCLE_COUNTER(STAT_FoliageAddInstanceBrush);
 
 	UWorld* World = GetWorld();
-	const bool bHasValidLandscapeLayers = LandscapeLayersValid(Settings);
+	const bool bHasValidLandscapeLayers = IsLandscapeLayersArrayValid(Settings->LandscapeLayers);
 
 	TArray<int32> ExistingInstanceBuckets;
 	ExistingInstanceBuckets.AddZeroed(NUM_INSTANCE_BUCKETS);
@@ -1684,7 +1704,7 @@ void FEdModeFoliage::AddInstancesForBrush(UWorld* InWorld, const UFoliageType* S
 				FFoliageInstance& Instance = FoliageInfo->Instances[Idx];
 				auto InstanceBasePtr = It.GetActor()->InstanceBaseCache.GetInstanceBasePtr(Instance.BaseId);
 				float HitWeight;
-				if (GetMaxHitWeight(Instance.Location, InstanceBasePtr.Get(), Settings, &LandscapeLayerCaches, HitWeight))
+				if (GetMaxHitWeight(Instance.Location, InstanceBasePtr.Get(), Settings->LandscapeLayers, &LandscapeLayerCaches, HitWeight))
 				{
 					// Add count to bucket.
 					ExistingInstanceBuckets[FMath::RoundToInt(HitWeight * (float)(NUM_INSTANCE_BUCKETS - 1))]++;
@@ -2447,7 +2467,7 @@ void FEdModeFoliage::ReapplyInstancesForBrush(UWorld* InWorld, AInstancedFoliage
 			}
 
 			// Find a ground normal for either normal or ground slope check.
-			if (bReapplyNormal || Settings->ReapplyGroundSlope || Settings->ReapplyVertexColorMask || (Settings->ReapplyLandscapeLayers && LandscapeLayersValid(Settings)))
+			if (bReapplyNormal || Settings->ReapplyGroundSlope || Settings->ReapplyVertexColorMask || (Settings->ReapplyLandscapeLayers && IsLandscapeLayersArrayValid(Settings->LandscapeLayers)))
 			{
 				FHitResult Hit;
 				static const FName NAME_ReapplyInstancesForBrush = TEXT("ReapplyInstancesForBrush");
@@ -2481,12 +2501,12 @@ void FEdModeFoliage::ReapplyInstancesForBrush(UWorld* InWorld, AInstancedFoliage
 					}
 
 					// Cull instances for the landscape layer
-					if (Settings->ReapplyLandscapeLayers && LandscapeLayersValid(Settings))
+					if (Settings->ReapplyLandscapeLayers && IsLandscapeLayersArrayValid(Settings->LandscapeLayers))
 					{
 						float HitWeight = 1.f;
-						if (GetMaxHitWeight(Hit.Location, Hit.GetComponent(), Settings, &LandscapeLayerCaches, HitWeight))
+						if (GetMaxHitWeight(Hit.Location, Hit.GetComponent(), Settings->LandscapeLayers, &LandscapeLayerCaches, HitWeight))
 						{
-							if (FilterByWeight(HitWeight, Settings))
+							if (IsFilteredByWeight(HitWeight, Settings->MinimumLayerWeight))
 							{
 								InstancesToDelete.Add(InstanceIndex);
 								if (bReapplyLocation)
diff --git a/Engine/Source/Runtime/Foliage/Public/FoliageType.h b/Engine/Source/Runtime/Foliage/Public/FoliageType.h
--- a/Engine/Source/Runtime/Foliage/Public/FoliageType.h
+++ b/Engine/Source/Runtime/Foliage/Public/FoliageType.h
@@ -210,10 +210,22 @@ public:
 	UPROPERTY(EditAnywhere, Category=Placement, meta=(ReapplyCondition="ReapplyHeight"))
 	FFloatInterval Height;
 
-	/** If a layer name is specified, painting on landscape will limit the foliage to areas of landscape with the specified layer painted */
-	UPROPERTY(EditAnywhere, AdvancedDisplay, Category=Placement, meta=(ReapplyCondition="ReapplyLandscapeLayers"))
+	/** If layer names are specified, painting on landscape will limit the foliage to areas of landscape with the specified layers painted */
+	UPROPERTY(EditAnywhere, AdvancedDisplay, Category=Placement, meta=(ReapplyCondition="ReapplyLandscapeLayers", DisplayName="Inclusion Landscape Layers"))
 	TArray<FName> LandscapeLayers;
 
+	/** Specifies the minimum value above which the landscape layer weight value must be, in order for foliage instances to be placed in a specific area */
+	UPROPERTY(EditAnywhere, AdvancedDisplay, Category=Placement, meta=(UIMin=0, ClampMin = 0, UIMax = 1, ClampMax = 1, HideBehind="LandscapeLayers", DisplayName="Minimum Inclusion Landscape Weight"))
+	float MinimumLayerWeight;
+
+	/** If layer names are specified, painting on landscape will exclude the foliage to areas of landscape without the specified layers painted */
+	UPROPERTY(EditAnywhere, AdvancedDisplay, Category = Placement, meta = (ReapplyCondition = "ReapplyLandscapeLayers"))
+	TArray<FName> ExclusionLandscapeLayers;
+
+	/** Specifies the minimum value above which the landscape exclusion layer weight value must be, in order for foliage instances to be excluded in a specific area */
+	UPROPERTY(EditAnywhere, AdvancedDisplay, Category = Placement, meta = (UIMin=0, ClampMin = 0, UIMax = 1, ClampMax = 1, HideBehind = "ExclusionLandscapeLayers"))
+	float MinimumExclusionLayerWeight;
+
 	UPROPERTY()
 	FName LandscapeLayer_DEPRECATED;
 	
@@ -225,10 +237,6 @@ public:
 	UPROPERTY(EditAnywhere, AdvancedDisplay, Category=Placement, meta=(HideBehind="CollisionWithWorld"))
 	FVector CollisionScale;
 
-	/** Specifies the minimum value above which the landscape layer weight value must be, in order for foliage instances to be placed in a specific area */
-	UPROPERTY(EditAnywhere, AdvancedDisplay, Category=Placement, meta=(HideBehind="LandscapeLayers"))
-	float MinimumLayerWeight;
-
 	UPROPERTY()
 	FBoxSphereBounds MeshBounds;
 
@@ -236,7 +244,6 @@ public:
 	UPROPERTY()
 	FVector LowBoundOriginRadius;
 
-
 public:
 	// INSTANCE SETTINGS
 
