Adjusted how the brush radius is adjusted to match the method used in landscape and mesh paint modes, ie. modifying by a percentage of the current size rather than a fixed amount.

Foliage mode keybindings can now be changed before the mode is activated.

diff --git a/Engine/Source/Editor/FoliageEdit/Private/FoliageEdMode.cpp b/Engine/Source/Editor/FoliageEdit/Private/FoliageEdMode.cpp
--- a/Engine/Source/Editor/FoliageEdit/Private/FoliageEdMode.cpp
+++ b/Engine/Source/Editor/FoliageEdit/Private/FoliageEdMode.cpp
@@ -233,7 +233,6 @@ FEdModeFoliage::FEdModeFoliage()
 	FName OpacityParamName("OpacityAmount");
 	BrushMID->GetScalarParameterValue(OpacityParamName, DefaultBrushOpacity);
 
-	FFoliageEditCommands::Register();
 	UICommandList = MakeShareable(new FUICommandList);
 	BindCommands();
 }
@@ -244,12 +243,32 @@ void FEdModeFoliage::BindCommands()
 
 	UICommandList->MapAction(
 		Commands.IncreaseBrushSize,
-		FExecuteAction::CreateRaw(this, &FEdModeFoliage::AdjustBrushRadius, 50.f),
+		FExecuteAction::CreateRaw(this, &FEdModeFoliage::AdjustBrushRadius, 1.f),
 		FCanExecuteAction::CreateRaw(this, &FEdModeFoliage::CurrentToolUsesBrush));
 
 	UICommandList->MapAction(
 		Commands.DecreaseBrushSize,
-		FExecuteAction::CreateRaw(this, &FEdModeFoliage::AdjustBrushRadius, -50.f),
+		FExecuteAction::CreateRaw(this, &FEdModeFoliage::AdjustBrushRadius, -1.f),
+		FCanExecuteAction::CreateRaw(this, &FEdModeFoliage::CurrentToolUsesBrush));
+
+	UICommandList->MapAction(
+		Commands.IncreasePaintDensity,
+		FExecuteAction::CreateRaw(this, &FEdModeFoliage::AdjustPaintDensity, 1.f),
+		FCanExecuteAction::CreateRaw(this, &FEdModeFoliage::CurrentToolUsesBrush));
+
+	UICommandList->MapAction(
+		Commands.DecreasePaintDensity,
+		FExecuteAction::CreateRaw(this, &FEdModeFoliage::AdjustPaintDensity, -1.f),
+		FCanExecuteAction::CreateRaw(this, &FEdModeFoliage::CurrentToolUsesBrush));
+
+	UICommandList->MapAction(
+		Commands.IncreaseUnpaintDensity,
+		FExecuteAction::CreateRaw(this, &FEdModeFoliage::AdjustUnpaintDensity, 1.f),
+		FCanExecuteAction::CreateRaw(this, &FEdModeFoliage::CurrentToolUsesBrush));
+
+	UICommandList->MapAction(
+		Commands.DecreaseUnpaintDensity,
+		FExecuteAction::CreateRaw(this, &FEdModeFoliage::AdjustUnpaintDensity, -1.f),
 		FCanExecuteAction::CreateRaw(this, &FEdModeFoliage::CurrentToolUsesBrush));
 
 	UICommandList->MapAction(
@@ -1983,23 +2002,44 @@ void FEdModeFoliage::SelectInvalidInstances(const UFoliageType* Settings)
 	}
 }
 
-void FEdModeFoliage::AdjustBrushRadius(float Adjustment)
+void FEdModeFoliage::AdjustBrushRadius(float Multiplier)
 {
 	if (UISettings.IsInAnySingleInstantiationMode())
 	{
 		return;
 	}
 
+	const float PercentageChange = 0.05f;
 	const float CurrentBrushRadius = UISettings.GetRadius();
 
-	if (Adjustment > 0.f)
+	float NewValue = CurrentBrushRadius * (1 + PercentageChange * Multiplier);
+	UISettings.SetRadius(FMath::Clamp(NewValue, 0.1f, 8192.0f));
+}
+
+void FEdModeFoliage::AdjustPaintDensity(float Multiplier)
+{
+	if (UISettings.IsInAnySingleInstantiationMode())
 	{
-		UISettings.SetRadius(FMath::Min(CurrentBrushRadius + Adjustment, 8192.f));
+		return;
 	}
-	else if (Adjustment < 0.f)
+
+	const float AdjustmentAmount = 0.02f;
+	const float CurrentDensity = UISettings.GetPaintDensity();
+
+	UISettings.SetPaintDensity(FMath::Clamp(CurrentDensity + AdjustmentAmount * Multiplier, 0.0f, 1.0f));
+}
+
+void FEdModeFoliage::AdjustUnpaintDensity(float Multiplier)
+{
+	if (UISettings.IsInAnySingleInstantiationMode())
 	{
-		UISettings.SetRadius(FMath::Max(CurrentBrushRadius + Adjustment, 0.f));
+		return;
 	}
+
+	const float AdjustmentAmount = 0.02f;
+	const float CurrentDensity = UISettings.GetUnpaintDensity();
+
+	UISettings.SetUnpaintDensity(FMath::Clamp(CurrentDensity + AdjustmentAmount * Multiplier, 0.0f, 1.0f));
 }
 
 void FEdModeFoliage::ReapplyInstancesForBrush(UWorld* InWorld, const UFoliageType* Settings, const FSphere& BrushSphere, float Pressure)
diff --git a/Engine/Source/Editor/FoliageEdit/Private/FoliageEdMode.h b/Engine/Source/Editor/FoliageEdit/Private/FoliageEdMode.h
--- a/Engine/Source/Editor/FoliageEdit/Private/FoliageEdMode.h
+++ b/Engine/Source/Editor/FoliageEdit/Private/FoliageEdMode.h
@@ -462,8 +462,14 @@ public:
 	/** Find and select instances that don't have valid base or 'off-ground' */
 	void SelectInvalidInstances(const UFoliageType* Settings);
 
-	/** Adjusts the radius of the foliage brush by the specified amount */
-	void AdjustBrushRadius(float Adjustment);
+	/** Adjusts the radius of the foliage brush, using the given multiplier to adjust speed */
+	void AdjustBrushRadius(float Multiplier);
+
+	/** Adjusts the painting density of the foliage brush, using the given multiplier to adjust speed */
+	void AdjustPaintDensity(float Multiplier);
+
+	/** Adjusts the unpainting (erasing) density of the foliage brush, using the given multiplier to adjust speed */
+	void AdjustUnpaintDensity(float Multiplier);
 
 	/** Add desired instances. Uses foliage settings to determine location/scale/rotation and whether instances should be ignored */
 	static void AddInstances(UWorld* InWorld, const TArray<FDesiredFoliageInstance>& DesiredInstances, const FFoliagePaintingGeometryFilter& OverrideGeometryFilter, bool InRebuildFoliageTree = true);
diff --git a/Engine/Source/Editor/FoliageEdit/Private/FoliageEditActions.cpp b/Engine/Source/Editor/FoliageEdit/Private/FoliageEditActions.cpp
--- a/Engine/Source/Editor/FoliageEdit/Private/FoliageEditActions.cpp
+++ b/Engine/Source/Editor/FoliageEdit/Private/FoliageEditActions.cpp
@@ -8,8 +8,14 @@
 
 void FFoliageEditCommands::RegisterCommands()
 {
-	UI_COMMAND(DecreaseBrushSize, "Decrease Brush Size", "Decreases the size of the foliage brush", EUserInterfaceActionType::Button, FInputChord(EKeys::LeftBracket));
-	UI_COMMAND(IncreaseBrushSize, "Increase Brush Size", "Increases the size of the foliage brush", EUserInterfaceActionType::Button, FInputChord(EKeys::RightBracket));
+	UI_COMMAND( DecreaseBrushSize, "Decrease Brush Size", "Decreases the size of the foliage brush", EUserInterfaceActionType::Button, FInputChord(EKeys::LeftBracket) );
+	UI_COMMAND( IncreaseBrushSize, "Increase Brush Size", "Increases the size of the foliage brush", EUserInterfaceActionType::Button, FInputChord(EKeys::RightBracket) );
+
+	UI_COMMAND( DecreasePaintDensity, "Decrease Brush Density", "Decreases the density of the foliage brush", EUserInterfaceActionType::Button, FInputChord(EModifierKey::Control, EKeys::LeftBracket) );
+	UI_COMMAND( IncreasePaintDensity, "Increase Brush Density", "Increases the density of the foliage brush", EUserInterfaceActionType::Button, FInputChord(EModifierKey::Control, EKeys::RightBracket) );
+
+	UI_COMMAND( DecreaseUnpaintDensity, "Decrease Erase Density", "Decreases the density of the foliage eraser", EUserInterfaceActionType::Button, FInputChord(EModifierKey::Control | EModifierKey::Shift, EKeys::LeftBracket));
+	UI_COMMAND( IncreaseUnpaintDensity, "Increase Erase Density", "Increases the density of the foliage eraser", EUserInterfaceActionType::Button, FInputChord(EModifierKey::Control | EModifierKey::Shift, EKeys::RightBracket));
 
 	UI_COMMAND( SetPaint, "Paint", "Paint", EUserInterfaceActionType::ToggleButton, FInputChord() );
 	UI_COMMAND( SetReapplySettings, "Reapply", "Reapply settings to instances", EUserInterfaceActionType::ToggleButton, FInputChord() );
diff --git a/Engine/Source/Editor/FoliageEdit/Private/FoliageEditActions.h b/Engine/Source/Editor/FoliageEdit/Private/FoliageEditActions.h
--- a/Engine/Source/Editor/FoliageEdit/Private/FoliageEditActions.h
+++ b/Engine/Source/Editor/FoliageEdit/Private/FoliageEditActions.h
@@ -31,6 +31,12 @@ public:
 	TSharedPtr< FUICommandInfo > IncreaseBrushSize;
 	TSharedPtr< FUICommandInfo > DecreaseBrushSize;
 
+	TSharedPtr< FUICommandInfo > IncreasePaintDensity;
+	TSharedPtr< FUICommandInfo > DecreasePaintDensity;
+
+	TSharedPtr< FUICommandInfo > IncreaseUnpaintDensity;
+	TSharedPtr< FUICommandInfo > DecreaseUnpaintDensity;
+
 	/** Commands for the tools toolbar. */
 	TSharedPtr< FUICommandInfo > SetPaint;
 	TSharedPtr< FUICommandInfo > SetReapplySettings;
diff --git a/Engine/Source/Editor/FoliageEdit/Private/FoliageEditModule.cpp b/Engine/Source/Editor/FoliageEdit/Private/FoliageEditModule.cpp
--- a/Engine/Source/Editor/FoliageEdit/Private/FoliageEditModule.cpp
+++ b/Engine/Source/Editor/FoliageEdit/Private/FoliageEditModule.cpp
@@ -19,6 +19,7 @@ const FName FoliageEditAppIdentifier = FName(TEXT("FoliageEdApp"));
 #include "FoliageType_Actor.h"
 #include "InstancedFoliageActor.h"
 #include "FoliageEdMode.h"
+#include "FoliageEditActions.h"
 #include "PropertyEditorModule.h"
 #include "FoliageTypeDetails.h"
 #include "ProceduralFoliageComponent.h"
@@ -52,6 +53,8 @@ public:
 			true, 400
 			);
 
+		FFoliageEditCommands::Register();
+
 		// Register the details customizer
 		FPropertyEditorModule& PropertyModule = FModuleManager::LoadModuleChecked<FPropertyEditorModule>("PropertyEditor");
 		PropertyModule.RegisterCustomClassLayout("FoliageType", FOnGetDetailCustomizationInstance::CreateStatic(&FFoliageTypeDetails::MakeInstance));
@@ -76,7 +79,7 @@ public:
 
 		SubscribeEvents();
 #endif
-
+		
 		// Register thumbnail renderer
 		UThumbnailManager::Get().RegisterCustomRenderer(UFoliageType_InstancedStaticMesh::StaticClass(), UFoliageType_ISMThumbnailRenderer::StaticClass());
 		UThumbnailManager::Get().RegisterCustomRenderer(UFoliageType_Actor::StaticClass(), UFoliageType_ActorThumbnailRenderer::StaticClass());
@@ -87,6 +90,8 @@ public:
 	 */
 	virtual void ShutdownModule() override
 	{
+		FFoliageEditCommands::Unregister();
+
 		FEditorModeRegistry::Get().UnregisterMode(FBuiltinEditorModes::EM_Foliage);
 
 		if (!UObjectInitialized())
