Fix foliage brush not updating when hovering over viewport

diff --git a/Engine/Source/Editor/FoliageEdit/Private/FoliageEdMode.cpp b/Engine/Source/Editor/FoliageEdit/Private/FoliageEdMode.cpp
--- a/Engine/Source/Editor/FoliageEdit/Private/FoliageEdMode.cpp
+++ b/Engine/Source/Editor/FoliageEdit/Private/FoliageEdMode.cpp
@@ -3963,7 +3963,7 @@ void FEdModeFoliage::ForceRealTimeViewports(const bool bEnable)
 		{
 			if (bEnable)
 			{
-				const bool bShouldBeRealtime = false;
+				const bool bShouldBeRealtime = true;
 				Viewport.SetRealtimeOverride(bShouldBeRealtime, LOCTEXT("RealtimeOverrideMessage_Foliage", "Foliage Mode"));
 			}
 			else
