Based on #PR 6828  Make sure we EndTransactions on ExitMode
Also wrap most Begin/EndTransactions inside scopes

diff --git a/Engine/Source/Editor/FoliageEdit/Private/FoliageEdMode.cpp b/Engine/Source/Editor/FoliageEdit/Private/FoliageEdMode.cpp
--- a/Engine/Source/Editor/FoliageEdit/Private/FoliageEdMode.cpp
+++ b/Engine/Source/Editor/FoliageEdit/Private/FoliageEdMode.cpp
@@ -233,6 +233,7 @@ FEdModeFoliage::FEdModeFoliage()
 	, UpdateSelectionCounter(0)
 	, bHasDeferredSelectionNotification(false)
 	, bMoving(false)
+	, bTracking(false)
 {
 	// Load resources and construct brush component
 	UStaticMesh* StaticMesh = nullptr;
@@ -475,6 +476,16 @@ void FEdModeFoliage::Enter()
 /** FEdMode: Called when the mode is exited */
 void FEdModeFoliage::Exit()
 {
+	if (bToolActive)
+	{
+		EndFoliageBrushTrace();
+	}
+	else if (bTracking)
+	{
+		EndTracking();
+	}
+	check(!GEditor->IsTransactionActive());
+
 	// Unregister VR mode from event handlers
 	{
 		UViewportWorldInteraction* ViewportWorldInteraction = Cast<UViewportWorldInteraction>( GEditor->GetEditorWorldExtensionsManager()->GetEditorWorldExtensions( GetWorld() )->FindExtension( UViewportWorldInteraction::StaticClass() ) );
@@ -672,7 +683,7 @@ void FEdModeFoliage::OnVRAction(class FEditorViewportClient& ViewportClient, UVi
 
 							if (HitResult.Actor.Get() != nullptr)
 							{
-								GEditor->BeginTransaction(NSLOCTEXT("UnrealEd", "FoliageMode_EditTransaction", "Foliage Editing"));
+								FScopedTransaction Transaction(NSLOCTEXT("UnrealEd", "FoliageMode_EditTransaction", "Foliage Editing"));
 
 								if (IsModifierButtonPressed(&ViewportClient))
 								{
@@ -682,8 +693,6 @@ void FEdModeFoliage::OnVRAction(class FEditorViewportClient& ViewportClient, UVi
 								{
 									ApplyPaintBucket_Add(HitResult.Actor.Get());
 								}
-
-								GEditor->EndTransaction();
 							}
 						}
 						// Select an instanced foliage
@@ -691,7 +700,7 @@ void FEdModeFoliage::OnVRAction(class FEditorViewportClient& ViewportClient, UVi
 						{
 							FHitResult HitResult = Interactor->GetHitResultFromLaserPointer();
 
-							GEditor->BeginTransaction(NSLOCTEXT("UnrealEd", "FoliageMode_EditTransaction", "Foliage Editing"));
+							FScopedTransaction Transaction(NSLOCTEXT("UnrealEd", "FoliageMode_EditTransaction", "Foliage Editing"));
 
 							if (HitResult.GetActor() != nullptr)
 							{
@@ -705,8 +714,6 @@ void FEdModeFoliage::OnVRAction(class FEditorViewportClient& ViewportClient, UVi
 								}
 							}
 
-							GEditor->EndTransaction();
-
 							// @todo vreditor: we currently don't have a key mapping scheme to snap selected instances to ground 
 							// SnapSelectedInstancesToGround(GetWorld());
 
@@ -2215,7 +2222,7 @@ void FEdModeFoliage::UpdateWidgetLocationToInstanceSelection()
 
 void FEdModeFoliage::RemoveSelectedInstances(UWorld* InWorld)
 {
-	GEditor->BeginTransaction(NSLOCTEXT("UnrealEd", "FoliageMode_EditTransaction", "Foliage Editing"));
+	FScopedTransaction Transaction(NSLOCTEXT("UnrealEd", "FoliageMode_EditTransaction", "Foliage Editing"));
 
 	const int32 NumLevels = InWorld->GetNumLevels();
 	for (int32 LevelIdx = 0; LevelIdx < NumLevels; ++LevelIdx)
@@ -2251,8 +2258,6 @@ void FEdModeFoliage::RemoveSelectedInstances(UWorld* InWorld)
 			}
 		}
 	}
-
-	GEditor->EndTransaction();
 }
 
 void FEdModeFoliage::GetSelectedInstanceFoliageTypes(TArray<const UFoliageType*>& OutFoliageTypes) const
@@ -3220,7 +3225,7 @@ bool FEdModeFoliage::GetStaticMeshVertexColorForHit(const UStaticMeshComponent*
 
 void FEdModeFoliage::SnapSelectedInstancesToGround(UWorld* InWorld)
 {
-	GEditor->BeginTransaction(NSLOCTEXT("UnrealEd", "FoliageMode_Transaction_SnapToGround", "Snap Foliage To Ground"));
+	FScopedTransaction Transaction(NSLOCTEXT("UnrealEd", "FoliageMode_Transaction_SnapToGround", "Snap Foliage To Ground"));
 	{
 		bool bMovedInstance = false;
 
@@ -3265,7 +3270,6 @@ void FEdModeFoliage::SnapSelectedInstancesToGround(UWorld* InWorld)
 			UpdateWidgetLocationToInstanceSelection();
 		}
 	}
-	GEditor->EndTransaction();
 }
 
 void FEdModeFoliage::HandleOnActorSpawned(AActor* Actor)
@@ -3701,13 +3705,13 @@ bool FEdModeFoliage::RemoveFoliageType(UFoliageType** FoliageTypeList, int32 Num
 
 	if (IFAList.Num())
 	{
-		GEditor->BeginTransaction(NSLOCTEXT("UnrealEd", "FoliageMode_RemoveMeshTransaction", "Foliage Editing: Remove Mesh"));
-		for (AInstancedFoliageActor* IFA : IFAList)
-		{
-			IFA->RemoveFoliageType(FoliageTypeList, Num);
+		{	// Transaction scope
+			FScopedTransaction Transaction(NSLOCTEXT("UnrealEd", "FoliageMode_RemoveMeshTransaction", "Foliage Editing: Remove Mesh"));
+			for (AInstancedFoliageActor* IFA : IFAList)
+			{
+				IFA->RemoveFoliageType(FoliageTypeList, Num);
+			}
 		}
-		GEditor->EndTransaction();
-
 		PopulateFoliageMeshList();
 		return true;
 	}
@@ -4019,7 +4023,7 @@ bool FEdModeFoliage::HandleClick(FEditorViewportClient* InViewportClient, HHitPr
 	{
 		if (HitProxy && HitProxy->IsA(HActor::StaticGetType()))
 		{
-			GEditor->BeginTransaction(NSLOCTEXT("UnrealEd", "FoliageMode_EditTransaction", "Foliage Editing"));
+			FScopedTransaction Transaction(NSLOCTEXT("UnrealEd", "FoliageMode_EditTransaction", "Foliage Editing"));
 			
 			if (IsModifierButtonPressed(InViewportClient))
 			{
@@ -4029,8 +4033,6 @@ bool FEdModeFoliage::HandleClick(FEditorViewportClient* InViewportClient, HHitPr
 			{
 				ApplyPaintBucket_Add(((HActor*)HitProxy)->Actor);
 			}
-
-			GEditor->EndTransaction();
 		}
 
 		return true;
@@ -4048,10 +4050,10 @@ FVector FEdModeFoliage::GetWidgetLocation() const
 /** FEdMode: Called when a mouse button is pressed */
 bool FEdModeFoliage::StartTracking(FEditorViewportClient* InViewportClient, FViewport* InViewport)
 {
+	bTracking = true;
 	if (IsCtrlDown(InViewport) && InViewport->KeyState(EKeys::MiddleMouseButton) && (UISettings.GetPaintToolSelected() || UISettings.GetReapplyToolSelected() || UISettings.GetLassoSelectToolSelected()))
 	{
 		bAdjustBrushRadius = true;
-		return true;
 	}
 	else if (UISettings.GetSelectToolSelected() || UISettings.GetLassoSelectToolSelected())
 	{
@@ -4059,27 +4061,43 @@ bool FEdModeFoliage::StartTracking(FEditorViewportClient* InViewportClient, FVie
 		UpdateWidgetLocationToInstanceSelection();
 
 		GEditor->BeginTransaction(NSLOCTEXT("UnrealEd", "FoliageMode_EditTransaction", "Foliage Editing"));
-
+		
 		bCanAltDrag = true;
-
-		return true;
 	}
-	return FEdMode::StartTracking(InViewportClient, InViewport);
+	else
+	{	
+		bTracking = false;
+		return FEdMode::StartTracking(InViewportClient, InViewport);
+	}
+
+	return true;
 }
 
-/** FEdMode: Called when the a mouse button is released */
-bool FEdModeFoliage::EndTracking(FEditorViewportClient* InViewportClient, FViewport* InViewport)
+bool FEdModeFoliage::EndTracking()
 {
-	if (!bAdjustBrushRadius && (UISettings.GetSelectToolSelected() || UISettings.GetLassoSelectToolSelected()))
+	bTracking = false;
+
+	if (bAdjustBrushRadius)
+	{
+		bAdjustBrushRadius = false;
+		return true;
+	}
+	else if (UISettings.GetSelectToolSelected() || UISettings.GetLassoSelectToolSelected())
 	{
 		UpdateInstancePartitioning(GetWorld());
 		PostTransformSelectedInstances(GetWorld());
 		GEditor->EndTransaction();
 		return true;
 	}
-	else
+		
+	return false;
+}
+
+/** FEdMode: Called when the a mouse button is released */
+bool FEdModeFoliage::EndTracking(FEditorViewportClient* InViewportClient, FViewport* InViewport)
+{
+	if (EndTracking())
 	{
-		bAdjustBrushRadius = false;
 		return true;
 	}
 
diff --git a/Engine/Source/Editor/FoliageEdit/Private/FoliageEdMode.h b/Engine/Source/Editor/FoliageEdit/Private/FoliageEdMode.h
--- a/Engine/Source/Editor/FoliageEdit/Private/FoliageEdMode.h
+++ b/Engine/Source/Editor/FoliageEdit/Private/FoliageEdMode.h
@@ -401,6 +401,9 @@ public:
 	/** Tell us if we can moves selected foliage instances to the target level. */
 	bool CanMoveSelectedFoliageToLevel(ULevel* InTargetLevel) const;
 
+	/** Ends tracking and end potential transaction */
+	bool EndTracking();
+
 	/** FEdMode: widget handling */
 	virtual FVector GetWidgetLocation() const override;
 	virtual bool AllowWidgetMove() override;
@@ -707,5 +710,8 @@ private:
 
 	/** When transforming instances */
 	bool bMoving;
+
+	/** Flag to know when we are tracking a transaction in mouse delta */
+	bool bTracking;
 };
 
