WorldPartition: Disable FoliageType_Actor
OFPA: Actor Foliage don't have external package

diff --git a/Engine/Source/Editor/FoliageEdit/Private/FoliageEdMode.cpp b/Engine/Source/Editor/FoliageEdit/Private/FoliageEdMode.cpp
--- a/Engine/Source/Editor/FoliageEdit/Private/FoliageEdMode.cpp
+++ b/Engine/Source/Editor/FoliageEdit/Private/FoliageEdMode.cpp
@@ -56,6 +56,8 @@
 
 #include "FoliageEditUtility.h"
 
+#include "WorldPartition/WorldPartitionSubsystem.h"
+
 #define LOCTEXT_NAMESPACE "FoliageEdMode"
 #define FOLIAGE_SNAP_TRACE (10000.f)
 
@@ -1994,6 +1996,20 @@ void FEdModeFoliage::RemoveSelectedInstances(UWorld* InWorld)
 	}
 }
 
+void FEdModeFoliage::GetFoliageTypeFilters(TArray<const UClass*>& OutFilters) const
+{
+	OutFilters.Add(UFoliageType_InstancedStaticMesh::StaticClass());
+	if (UWorldPartitionSubsystem* WorldPartition = GetWorld()->GetSubsystem<UWorldPartitionSubsystem>())
+	{
+		if (WorldPartition->IsEnabled())
+		{
+			return;
+		}
+	}
+	// FoliageType Actor only supported in non WorldPartition worlds
+	OutFilters.Add(UFoliageType_Actor::StaticClass());
+}
+
 void FEdModeFoliage::GetSelectedInstanceFoliageTypes(TArray<const UFoliageType*>& OutFoliageTypes) const
 {
 	for (TActorIterator<AInstancedFoliageActor> It(GetWorld()); It; ++It)
diff --git a/Engine/Source/Editor/FoliageEdit/Private/FoliageEdMode.h b/Engine/Source/Editor/FoliageEdit/Private/FoliageEdMode.h
--- a/Engine/Source/Editor/FoliageEdit/Private/FoliageEdMode.h
+++ b/Engine/Source/Editor/FoliageEdit/Private/FoliageEdMode.h
@@ -562,6 +562,9 @@ public:
 
 	/** Remove currently selected instances*/
 	void RemoveSelectedInstances(UWorld* InWorld);
+
+	/** Returns the list of valid FoliageType class filters */
+	void GetFoliageTypeFilters(TArray<const UClass*>& OutFilters) const;
 			
 private:
 
diff --git a/Engine/Source/Editor/FoliageEdit/Private/SFoliagePalette.cpp b/Engine/Source/Editor/FoliageEdit/Private/SFoliagePalette.cpp
--- a/Engine/Source/Editor/FoliageEdit/Private/SFoliagePalette.cpp
+++ b/Engine/Source/Editor/FoliageEdit/Private/SFoliagePalette.cpp
@@ -683,7 +683,8 @@ void SFoliagePalette::OnSearchTextChanged(const FText& InFilterText)
 TSharedRef<SWidget> SFoliagePalette::GetAddFoliageTypePicker()
 {
 	TArray<const UClass*> ClassFilters;
-	ClassFilters.Add(UFoliageType::StaticClass());
+
+	FoliageEditMode->GetFoliageTypeFilters(ClassFilters);
 
 	return PropertyCustomizationHelpers::MakeAssetPickerWithMenu(FAssetData(),
 		false,
@@ -867,12 +868,23 @@ EVisibility SFoliagePalette::GetFoliageDropTargetVisibility() const
 {
 	if (FSlateApplication::Get().IsDragDropping())
 	{
+		TArray<const UClass*> ClassFilters;
+		FoliageEditMode->GetFoliageTypeFilters(ClassFilters);
+		// Support StaticMesh on drag
+		ClassFilters.Add(UStaticMesh::StaticClass());
+
 		TArray<FAssetData> DraggedAssets = AssetUtil::ExtractAssetDataFromDrag(FSlateApplication::Get().GetDragDroppingContent());
 		for (const FAssetData& AssetData : DraggedAssets)
 		{
-			if (AssetData.IsValid() && (AssetData.GetClass()->IsChildOf(UStaticMesh::StaticClass()) || AssetData.GetClass()->IsChildOf(UFoliageType::StaticClass())))
+			if (AssetData.IsValid())
 			{
-				return EVisibility::Visible;
+				for (const UClass* Filter : ClassFilters)
+				{
+					if (AssetData.GetClass()->IsChildOf(Filter))
+					{
+						return EVisibility::Visible;
+					}
+				}
 			}
 		}
 	}
diff --git a/Engine/Source/Runtime/Engine/Private/Level.cpp b/Engine/Source/Runtime/Engine/Private/Level.cpp
--- a/Engine/Source/Runtime/Engine/Private/Level.cpp
+++ b/Engine/Source/Runtime/Engine/Private/Level.cpp
@@ -55,6 +55,7 @@ Level.cpp: Level-related functions
 #include "AssetRegistryModule.h"
 #include "IAssetRegistry.h"
 #include "AssetData.h"
+#include "FoliageHelper.h"
 #endif
 #include "Engine/LevelStreaming.h"
 #include "LevelUtils.h"
@@ -2275,6 +2276,11 @@ bool ULevel::CanConvertActorToExternalPackaging(AActor* Actor)
 		return false;
 	}
 
+	if (FFoliageHelper::IsOwnedByFoliage(Actor))
+	{
+		return false;
+	}
+
 	return Actor->SupportsExternalPackaging();
 }
 
diff --git a/Engine/Source/Runtime/Foliage/Private/FoliageActor.cpp b/Engine/Source/Runtime/Foliage/Private/FoliageActor.cpp
--- a/Engine/Source/Runtime/Foliage/Private/FoliageActor.cpp
+++ b/Engine/Source/Runtime/Foliage/Private/FoliageActor.cpp
@@ -85,7 +85,7 @@ AActor* FFoliageActor::Spawn(AInstancedFoliageActor* IFA, const FFoliageInstance
 	FActorSpawnParameters SpawnParameters;
 	SpawnParameters.ObjectFlags = RF_Transactional;
 	SpawnParameters.bHideFromSceneOutliner = true;
-	SpawnParameters.bCreateActorPackage = true;
+	SpawnParameters.bCreateActorPackage = false; // No OFPA because that would generate tons of files
 	SpawnParameters.OverrideLevel = IFA->GetLevel();
 	AActor* NewActor = IFA->GetWorld()->SpawnActor(ActorClass, nullptr, nullptr, SpawnParameters);
 	if (NewActor)
