Interface for finding relevant actors using ActorPartitionSubsystem

diff --git a/Engine/Source/Editor/FoliageEdit/Private/FoliageEdMode.cpp b/Engine/Source/Editor/FoliageEdit/Private/FoliageEdMode.cpp
--- a/Engine/Source/Editor/FoliageEdit/Private/FoliageEdMode.cpp
+++ b/Engine/Source/Editor/FoliageEdit/Private/FoliageEdMode.cpp
@@ -623,6 +623,26 @@ void FEdModeFoliage::ClearAllToolSelection()
 	UISettings.SetIsInSingleInstantiationMode(false);
 }
 
+void FEdModeFoliage::ForEachFoliageInfo(UWorld* InWorld, const UFoliageType* FoliageType, const FSphere& BrushSphere, TFunctionRef<bool(AInstancedFoliageActor* IFA, FFoliageInfo* FoliageInfo, const UFoliageType* FoliageType)> InOperation)
+{
+	auto IFAOperation = [&FoliageType, &InOperation](APartitionActor* Actor)
+	{
+		AInstancedFoliageActor* IFA = Cast<AInstancedFoliageActor>(Actor);
+		FFoliageInfo* FoliageInfo = IFA ? IFA->FindInfo(FoliageType) : nullptr;
+		if (FoliageInfo)
+		{
+			return InOperation(IFA, FoliageInfo, FoliageType);
+		}
+		return true;
+	};
+
+	if (UActorPartitionSubsystem* ActorPartitionSubsystem = InWorld->GetSubsystem<UActorPartitionSubsystem>())
+	{	
+		const FBox BrushSphereBounds(BrushSphere.Center - BrushSphere.W, BrushSphere.Center + BrushSphere.W);
+		ActorPartitionSubsystem->ForEachRelevantActor(AInstancedFoliageActor::StaticClass(), BrushSphereBounds, IFAOperation);
+	}
+}
+
 void FEdModeFoliage::OnSetPaint()
 {
 	ClearAllToolSelection();
@@ -969,16 +989,14 @@ static bool CheckForOverlappingSphere(AInstancedFoliageActor* IFA, const UFoliag
 // Returns whether or not there is are any instances overlapping the sphere specified
 static bool CheckForOverlappingSphere(UWorld* InWorld, const UFoliageType* Settings, const FSphere& Sphere)
 {
-	for (FFoliageInfoIterator It(InWorld, Settings); It; ++It)
-	{
-		FFoliageInfo* Info = (*It);
-		if (Info->CheckForOverlappingSphere(Sphere))
-		{
-			return true;
-		}
-	}
+	bool bIsOverlappingSphere = false;
+	auto CheckForOverlap = [&bIsOverlappingSphere, &Sphere](AInstancedFoliageActor* IFA, FFoliageInfo* FoliageInfo, const UFoliageType* FoliageType) {
+		bIsOverlappingSphere = FoliageInfo->CheckForOverlappingSphere(Sphere);
+		return !bIsOverlappingSphere;
+	};
 
-	return false;
+	FEdModeFoliage::ForEachFoliageInfo(InWorld, Settings, Sphere, CheckForOverlap);
+	return bIsOverlappingSphere;
 }
 
 static bool CheckLocationForPotentialInstance(UWorld* InWorld, const UFoliageType* Settings, const bool bSingleInstanceMode, const FVector& Location, const FVector& Normal, TArray<FVector>& PotentialInstanceLocations, FFoliageInstanceHash& PotentialInstanceHash)
@@ -1490,10 +1508,9 @@ void FEdModeFoliage::AddInstancesForBrush(UWorld* InWorld, const UFoliageType* S
 	TArray<int32> ExistingInstanceBuckets;
 	ExistingInstanceBuckets.AddZeroed(NUM_INSTANCE_BUCKETS);
 	int32 NumExistingInstances = 0;
-
-	for (FFoliageInfoIterator It(World, Settings); It; ++It)
-	{
-		FFoliageInfo* FoliageInfo = (*It);
+	
+	auto AddingInstances = [this, &BrushSphere, &NumExistingInstances, &bHasValidLandscapeLayers, &ExistingInstanceBuckets](AInstancedFoliageActor* IFA, FFoliageInfo* FoliageInfo, const UFoliageType* FoliageType) {
+	
 		TArray<int32> ExistingInstances;
 		FoliageInfo->GetInstancesInsideSphere(BrushSphere, ExistingInstances);
 		NumExistingInstances += ExistingInstances.Num();
@@ -1504,9 +1521,9 @@ void FEdModeFoliage::AddInstancesForBrush(UWorld* InWorld, const UFoliageType* S
 			for (int32 Idx : ExistingInstances)
 			{
 				FFoliageInstance& Instance = FoliageInfo->Instances[Idx];
-				auto InstanceBasePtr = It.GetActor()->InstanceBaseCache.GetInstanceBasePtr(Instance.BaseId);
+				auto InstanceBasePtr = IFA->InstanceBaseCache.GetInstanceBasePtr(Instance.BaseId);
 				float HitWeight;
-				if (GetMaxHitWeight(Instance.Location, InstanceBasePtr.Get(), Settings->LandscapeLayers, &LandscapeLayerCaches, HitWeight))
+				if (GetMaxHitWeight(Instance.Location, InstanceBasePtr.Get(), FoliageType->LandscapeLayers, &LandscapeLayerCaches, HitWeight))
 				{
 					// Add count to bucket.
 					ExistingInstanceBuckets[FMath::RoundToInt(HitWeight * (float)(NUM_INSTANCE_BUCKETS - 1))]++;
@@ -1518,7 +1535,11 @@ void FEdModeFoliage::AddInstancesForBrush(UWorld* InWorld, const UFoliageType* S
 			// When not tied to a layer, put all the ExistingInstances in the last bucket.
 			ExistingInstanceBuckets[NUM_INSTANCE_BUCKETS - 1] = NumExistingInstances;
 		}
-	}
+
+		return true;
+	};
+	
+	ForEachFoliageInfo(InWorld, Settings, BrushSphere, AddingInstances);
 
 	if (DesiredInstanceCount > NumExistingInstances)
 	{
@@ -1541,22 +1562,18 @@ void FEdModeFoliage::RemoveInstancesForBrush(UWorld* InWorld, const UFoliageType
 {
 	SCOPE_CYCLE_COUNTER(STAT_FoliageRemoveInstanceBrush);
 
-	for (FFoliageInfoIterator It(InWorld, Settings); It; ++It)
-	{
-		FFoliageInfo* FoliageInfo = (*It);
-		AInstancedFoliageActor* IFA = It.GetActor();
-
+	auto RemovingInstances = [this, &BrushSphere, &DesiredInstanceCount, &Pressure](AInstancedFoliageActor* IFA, FFoliageInfo* FoliageInfo, const UFoliageType* FoliageType) {
 		TArray<int32> PotentialInstancesToRemove;
 		FoliageInfo->GetInstancesInsideSphere(BrushSphere, PotentialInstancesToRemove);
 		if (PotentialInstancesToRemove.Num() == 0)
 		{
-			continue;
+			return true;
 		}
 
 		int32 InstancesToRemove = FMath::RoundToInt((float)(PotentialInstancesToRemove.Num() - DesiredInstanceCount) * Pressure);
 		if (InstancesToRemove <= 0)
 		{
-			continue;
+			return true;
 		}
 
 		int32 InstancesToKeep = PotentialInstancesToRemove.Num() - InstancesToRemove;
@@ -1594,47 +1611,54 @@ void FEdModeFoliage::RemoveInstancesForBrush(UWorld* InWorld, const UFoliageType
 
 			FoliageInfo->RemoveInstances(IFA, PotentialInstancesToRemove, false);
 		}
-	}
+		return true;
+	};
+	ForEachFoliageInfo(InWorld, Settings, BrushSphere, RemovingInstances);
 }
 
 
 void FEdModeFoliage::SelectInstanceAtLocation(UWorld* InWorld, const UFoliageType* Settings, const FVector& Location, bool bSelect)
 {
 	FEdModeFoliageSelectionUpdate Scope(this);
-	for (FFoliageInfoIterator It(InWorld, Settings); It; ++It)
-	{
-		FFoliageInfo* FoliageInfo = (*It);
-		AInstancedFoliageActor* IFA = It.GetActor();
+	const float WidthOfSphere = 10.f;
+	FSphere SphereBounds(Location,WidthOfSphere);
+
+	auto SelectInstances = [&Location, &bSelect](AInstancedFoliageActor* IFA, FFoliageInfo* FoliageInfo, const UFoliageType* FoliageType) {
 
 		int32 Instance;
 		bool bResult;
 		FoliageInfo->GetInstanceAtLocation(Location, Instance, bResult);
 		if (bResult)
 		{
-		    TArray<int32> Instances;
+			TArray<int32> Instances;
 			Instances.Add(Instance);
 			FoliageInfo->SelectInstances(IFA, bSelect, Instances);
 		}
-	}
+
+		return true;
+	};
+	ForEachFoliageInfo(InWorld, Settings, SphereBounds, SelectInstances);
+	
 }
 
 void FEdModeFoliage::SelectInstancesForBrush(UWorld* InWorld, const UFoliageType* Settings, const FSphere& BrushSphere, bool bSelect)
 {
 	FEdModeFoliageSelectionUpdate Scope(this);
-	for (FFoliageInfoIterator It(InWorld, Settings); It; ++It)
-	{
-		FFoliageInfo* FoliageInfo = (*It);
-		AInstancedFoliageActor* IFA = It.GetActor();
+	
+	auto SelectInstances = [&BrushSphere, &bSelect](AInstancedFoliageActor* IFA, FFoliageInfo* FoliageInfo, const UFoliageType* FoliageType) {
 
 		TArray<int32> Instances;
 		FoliageInfo->GetInstancesInsideSphere(BrushSphere, Instances);
 		if (Instances.Num() == 0)
 		{
-			continue;
+			return true;
 		}
-
 		FoliageInfo->SelectInstances(IFA, bSelect, Instances);
-	}
+
+		return true;
+	};
+
+	ForEachFoliageInfo(InWorld, Settings, BrushSphere, SelectInstances);
 }
 
 void RefreshSceneOutliner()
@@ -2177,13 +2201,12 @@ void FEdModeFoliage::ReapplyInstancesForBrush(UWorld* InWorld, const UFoliageTyp
 	// Adjust instance density first
 	ReapplyInstancesDensityForBrush(InWorld, Settings, BrushSphere, Pressure);
 
-	for (FFoliageInfoIterator It(InWorld, Settings); It; ++It)
-	{
-		FFoliageInfo* FoliageInfo = (*It);
-		AInstancedFoliageActor* IFA = It.GetActor();
-
-		ReapplyInstancesForBrush(InWorld, IFA, Settings, FoliageInfo, BrushSphere, Pressure, bSingleInstanceMode);
-	}
+	auto ReapplyInstances = [this, &InWorld, &BrushSphere, &Pressure, &bSingleInstanceMode](AInstancedFoliageActor* IFA, FFoliageInfo* FoliageInfo, const UFoliageType* FoliageType) {
+		
+		ReapplyInstancesForBrush(InWorld, IFA, FoliageType, FoliageInfo, BrushSphere, Pressure, bSingleInstanceMode);
+		return true;
+	};
+	ForEachFoliageInfo(InWorld, Settings, BrushSphere, ReapplyInstances);
 }
 
 /** Reapply instance settings to exiting instances */
diff --git a/Engine/Source/Editor/FoliageEdit/Private/FoliageEdMode.h b/Engine/Source/Editor/FoliageEdit/Private/FoliageEdMode.h
--- a/Engine/Source/Editor/FoliageEdit/Private/FoliageEdMode.h
+++ b/Engine/Source/Editor/FoliageEdit/Private/FoliageEdMode.h
@@ -568,6 +568,9 @@ public:
 
 	/** Returns the list of valid FoliageType class filters */
 	void GetFoliageTypeFilters(TArray<const UClass*>& OutFilters) const;
+
+	/*Find the relevant foliage actor with the foliage type and run the operation*/
+	static void ForEachFoliageInfo(UWorld* InWorld, const UFoliageType* FoliageType, const FSphere& BrushSphere, TFunctionRef<bool(AInstancedFoliageActor* IFA, FFoliageInfo* FoliageInfo, const UFoliageType* FoliageType)> InOperation);
 			
 private:
 
diff --git a/Engine/Source/Runtime/Engine/Private/ActorPartition/ActorPartitionSubsystem.cpp b/Engine/Source/Runtime/Engine/Private/ActorPartition/ActorPartitionSubsystem.cpp
--- a/Engine/Source/Runtime/Engine/Private/ActorPartition/ActorPartitionSubsystem.cpp
+++ b/Engine/Source/Runtime/Engine/Private/ActorPartition/ActorPartitionSubsystem.cpp
@@ -5,6 +5,7 @@
 #include "WorldPartition/ActorPartition/PartitionActorDesc.h"
 #include "WorldPartition/WorldPartitionSubsystem.h"
 #include "Engine/World.h"
+#include "EngineUtils.h"
 
 DEFINE_LOG_CATEGORY_STATIC(LogActorPartitionSubsystem, All, All);
 
@@ -125,6 +126,19 @@ public:
 		return FoundActor;
 	}
 
+
+	void ForEachRelevantActor(const TSubclassOf<APartitionActor>& InActorClass, const FBox& IntersectionBounds, TFunctionRef<bool(APartitionActor*)>InOperation) const override
+	{
+		for (TActorIterator<APartitionActor> It(World, InActorClass); It; ++It)
+		{
+			if (!InOperation(*It))
+			{
+				return;
+			}
+		}
+	}
+
+
 private:
 	void OnLevelRemovedFromWorld(ULevel* InLevel, UWorld* InWorld)
 	{
@@ -255,6 +269,23 @@ public:
 		return FoundActor;
 	}
 
+	void ForEachRelevantActor(const TSubclassOf<APartitionActor>& InActorClass, const FBox& IntersectionBounds, TFunctionRef<bool(APartitionActor*)>InOperation) const override
+	{
+		UActorPartitionSubsystem* ActorSubsystem = World->GetSubsystem<UActorPartitionSubsystem>();
+		FActorPartitionGridHelper::ForEachIntersectingCell(InActorClass, IntersectionBounds, World->PersistentLevel, [&ActorSubsystem, &InActorClass, &IntersectionBounds, &InOperation](const UActorPartitionSubsystem::FCellCoord& InCellCoord, const FBox& InCellBounds) {
+
+			if (InCellBounds.Intersect(IntersectionBounds))
+			{
+				const bool bCreate = false;
+				if (auto PartitionActor = ActorSubsystem->GetActor(InActorClass, InCellCoord, bCreate))
+				{
+					return InOperation(PartitionActor);
+				}
+			}
+			return true;
+			});
+	}
+
 private:
 	UWorldPartitionSubsystem* WorldPartition;
 };
@@ -291,6 +322,14 @@ void UActorPartitionSubsystem::Deinitialize()
 	}
 }
 
+void UActorPartitionSubsystem::ForEachRelevantActor(const TSubclassOf<APartitionActor>& InActorClass, const FBox& IntersectionBounds, TFunctionRef<bool(APartitionActor*)>InOperation) const
+{
+	if (ActorPartition)
+	{
+		ActorPartition->ForEachRelevantActor(InActorClass, IntersectionBounds, InOperation);
+	}
+}
+
 void UActorPartitionSubsystem::OnActorPartitionHashInvalidated(const FCellCoord& Hash)
 {
 	PartitionedActors.Remove(Hash);
diff --git a/Engine/Source/Runtime/Engine/Public/ActorPartition/ActorPartitionSubsystem.h b/Engine/Source/Runtime/Engine/Public/ActorPartition/ActorPartitionSubsystem.h
--- a/Engine/Source/Runtime/Engine/Public/ActorPartition/ActorPartitionSubsystem.h
+++ b/Engine/Source/Runtime/Engine/Public/ActorPartition/ActorPartitionSubsystem.h
@@ -135,6 +135,8 @@ public:
 
 	virtual void Initialize(FSubsystemCollectionBase& Collection) override;
 	virtual void Deinitialize() override;
+
+	void ForEachRelevantActor(const TSubclassOf<APartitionActor>& InActorClass, const FBox& IntersectionBounds, TFunctionRef<bool(APartitionActor*)>InOperation) const;
 #endif
 	bool IsLevelPartition() const;
 
@@ -166,6 +168,7 @@ public:
 
 	virtual UActorPartitionSubsystem::FCellCoord GetActorPartitionHash(const FActorPartitionGetParams& GetParams) const = 0;
 	virtual APartitionActor* GetActor(const TSubclassOf<APartitionActor>& InActorClass, bool bInCreate, const UActorPartitionSubsystem::FCellCoord& InCellCoord, const FGuid& InGuid, uint32 InGridSize, bool bInBoundsSearch, TFunctionRef<void(APartitionActor*)> InActorCreated) = 0;
+	virtual void ForEachRelevantActor(const TSubclassOf<APartitionActor>& InActorClass, const FBox& IntersectionBounds, TFunctionRef<bool(APartitionActor*)>InOperation) const = 0;
 
 	DECLARE_EVENT_OneParam(FBaseActorPartition, FOnActorPartitionHashInvalidated, const UActorPartitionSubsystem::FCellCoord&);
 	FOnActorPartitionHashInvalidated& GetOnActorPartitionHashInvalidated() { return OnActorPartitionHashInvalidated; }
