Foliage: Expose PFV Resimulate/Clear to BP

diff --git a/Engine/Source/Editor/FoliageEdit/Private/ProceduralFoliageComponentDetails.cpp b/Engine/Source/Editor/FoliageEdit/Private/ProceduralFoliageComponentDetails.cpp
--- a/Engine/Source/Editor/FoliageEdit/Private/ProceduralFoliageComponentDetails.cpp
+++ b/Engine/Source/Editor/FoliageEdit/Private/ProceduralFoliageComponentDetails.cpp
@@ -11,10 +11,8 @@
 #include "InstancedFoliage.h"
 #include "ProceduralFoliageSpawner.h"
 #include "ProceduralFoliageComponent.h"
-#include "FoliageEdMode.h"
+#include "ProceduralFoliageEditorLibrary.h"
 #include "ScopedTransaction.h"
-#include "Framework/Notifications/NotificationManager.h"
-#include "Widgets/Notifications/SNotificationList.h"
 
 #define LOCTEXT_NAMESPACE "ProceduralFoliageComponentDetails"
 
@@ -72,41 +70,18 @@ void FProceduralFoliageComponentDetails::CustomizeDetails(IDetailLayoutBuilder&
 
 FReply FProceduralFoliageComponentDetails::OnResimulateClicked()
 {
+	TArray<UProceduralFoliageComponent*> SelectedFoliageComponents;
+	SelectedFoliageComponents.Reserve(SelectedComponents.Num());
 	for (TWeakObjectPtr<UProceduralFoliageComponent>& Component : SelectedComponents)
 	{
-		if (Component.IsValid() && Component->FoliageSpawner)
+		if (Component.IsValid())
 		{
-			FScopedTransaction Transaction(LOCTEXT("Resimulate_Transaction", "Procedural Foliage Simulation"));
-			TArray <FDesiredFoliageInstance> DesiredFoliageInstances;
-			if (Component->GenerateProceduralContent(DesiredFoliageInstances))
-			{
-				if (DesiredFoliageInstances.Num() > 0)
-				{
-					Component->RemoveProceduralContent(false);
-
-					FFoliagePaintingGeometryFilter OverrideGeometryFilter;
-					OverrideGeometryFilter.bAllowLandscape = Component->bAllowLandscape;
-					OverrideGeometryFilter.bAllowStaticMesh = Component->bAllowStaticMesh;
-					OverrideGeometryFilter.bAllowBSP = Component->bAllowBSP;
-					OverrideGeometryFilter.bAllowFoliage = Component->bAllowFoliage;
-					OverrideGeometryFilter.bAllowTranslucent = Component->bAllowTranslucent;
-
-					FEdModeFoliage::AddInstances(Component->GetWorld(), DesiredFoliageInstances, OverrideGeometryFilter, true);					
-				}
-
-				// If no instances were spawned, inform the user
-				if (!Component->HasSpawnedAnyInstances())
-				{
-					FNotificationInfo Info(LOCTEXT("NothingSpawned_Notification", "Unable to spawn instances. Ensure a large enough surface exists within the volume."));
-					Info.bUseLargeFont = false;
-					Info.bFireAndForget = true;
-					Info.bUseThrobber = false;
-					Info.bUseSuccessFailIcons = true;
-					FSlateNotificationManager::Get().AddNotification(Info);
-				}
-			}
+			SelectedFoliageComponents.Add(Component.Get());
 		}
 	}
+	FScopedTransaction Transaction(LOCTEXT("Resimulate_Transaction", "Procedural Foliage Simulation"));
+	UProceduralFoliageEditorLibrary::ResimulateProceduralFoliageComponents(SelectedFoliageComponents);
+
 	return FReply::Handled();
 }
 
diff --git a/Engine/Source/Editor/FoliageEdit/Private/ProceduralFoliageEditorLibrary.cpp b/Engine/Source/Editor/FoliageEdit/Private/ProceduralFoliageEditorLibrary.cpp
new file mode 100644
--- /dev/null
+++ b/Engine/Source/Editor/FoliageEdit/Private/ProceduralFoliageEditorLibrary.cpp
@@ -0,0 +1,90 @@
+// Copyright Epic Games, Inc. All Rights Reserved.
+
+#include "ProceduralFoliageEditorLibrary.h"
+#include "ProceduralFoliageVolume.h"
+#include "ProceduralFoliageComponent.h"
+#include "FoliageEdMode.h"
+#include "Algo/Transform.h"
+#include "Framework/Notifications/NotificationManager.h"
+#include "Widgets/Notifications/SNotificationList.h"
+#include "Misc/ScopedSlowTask.h"
+
+#define LOCTEXT_NAMESPACE "ProceduralFoliageEditorLibrary"
+
+UProceduralFoliageEditorLibrary::UProceduralFoliageEditorLibrary(const FObjectInitializer& ObjectInitializer)
+	: Super(ObjectInitializer)
+{
+}
+
+void UProceduralFoliageEditorLibrary::ResimulateProceduralFoliageVolumes(const TArray<AProceduralFoliageVolume*>& ProceduralFoliageVolumes)
+{
+	TArray<UProceduralFoliageComponent*> Components;
+	Components.Reserve(ProceduralFoliageVolumes.Num());
+	Algo::Transform(ProceduralFoliageVolumes, Components, [](AProceduralFoliageVolume* PFV) { return PFV->ProceduralComponent; });
+	ResimulateProceduralFoliageComponents(Components);
+}
+
+void UProceduralFoliageEditorLibrary::ClearProceduralFoliageVolumes(const TArray<AProceduralFoliageVolume*>& ProceduralFoliageVolumes)
+{
+	TArray<UProceduralFoliageComponent*> Components;
+	Components.Reserve(ProceduralFoliageVolumes.Num());
+	Algo::Transform(ProceduralFoliageVolumes, Components, [](AProceduralFoliageVolume* PFV) { return PFV->ProceduralComponent; });
+	ClearProceduralFoliageComponents(Components);
+}
+
+void UProceduralFoliageEditorLibrary::ClearProceduralFoliageComponents(const TArray<UProceduralFoliageComponent*>& ProceduralFoliageComponents)
+{
+	FScopedSlowTask SlowTask(ProceduralFoliageComponents.Num(), LOCTEXT("ClearProceduralFoliageComponents", "Clearing Procedural Foliage Components"));
+	SlowTask.MakeDialogDelayed(0.5f);
+	for (UProceduralFoliageComponent* Component : ProceduralFoliageComponents)
+	{
+		SlowTask.EnterProgressFrame(1);
+		if (Component)
+		{
+			Component->RemoveProceduralContent();
+		}
+	}
+}
+
+void UProceduralFoliageEditorLibrary::ResimulateProceduralFoliageComponents(const TArray<UProceduralFoliageComponent*>& ProceduralFoliageComponents)
+{
+	FScopedSlowTask SlowTask(ProceduralFoliageComponents.Num(), LOCTEXT("ResimulateProceduralFoliageComponents", "Resimulating Procedural Foliage Components"));
+	SlowTask.MakeDialogDelayed(0.5f);
+	for (UProceduralFoliageComponent* Component : ProceduralFoliageComponents)
+	{
+		SlowTask.EnterProgressFrame(1);
+		if (Component && Component->FoliageSpawner)
+		{
+			TArray <FDesiredFoliageInstance> DesiredFoliageInstances;
+			if (Component->GenerateProceduralContent(DesiredFoliageInstances))
+			{
+				if (DesiredFoliageInstances.Num() > 0)
+				{
+					Component->RemoveProceduralContent(false);
+
+					FFoliagePaintingGeometryFilter OverrideGeometryFilter;
+					OverrideGeometryFilter.bAllowLandscape = Component->bAllowLandscape;
+					OverrideGeometryFilter.bAllowStaticMesh = Component->bAllowStaticMesh;
+					OverrideGeometryFilter.bAllowBSP = Component->bAllowBSP;
+					OverrideGeometryFilter.bAllowFoliage = Component->bAllowFoliage;
+					OverrideGeometryFilter.bAllowTranslucent = Component->bAllowTranslucent;
+
+					FEdModeFoliage::AddInstances(Component->GetWorld(), DesiredFoliageInstances, OverrideGeometryFilter, true);
+				}
+
+				// If no instances were spawned, inform the user
+				if (!Component->HasSpawnedAnyInstances())
+				{
+					FNotificationInfo Info(LOCTEXT("NothingSpawned_Notification", "Unable to spawn instances. Ensure a large enough surface exists within the volume."));
+					Info.bUseLargeFont = false;
+					Info.bFireAndForget = true;
+					Info.bUseThrobber = false;
+					Info.bUseSuccessFailIcons = true;
+					FSlateNotificationManager::Get().AddNotification(Info);
+				}
+			}
+		}
+	}
+}
+
+#undef LOCTEXT_NAMESPACE
\ No newline at end of file
diff --git a/Engine/Source/Editor/FoliageEdit/Public/ProceduralFoliageEditorLibrary.h b/Engine/Source/Editor/FoliageEdit/Public/ProceduralFoliageEditorLibrary.h
new file mode 100644
--- /dev/null
+++ b/Engine/Source/Editor/FoliageEdit/Public/ProceduralFoliageEditorLibrary.h
@@ -0,0 +1,33 @@
+// Copyright Epic Games, Inc. All Rights Reserved.
+
+/**
+ *
+ */
+
+#pragma once
+
+#include "Kismet/BlueprintFunctionLibrary.h"
+#include "Engine/EngineTypes.h"
+#include "ProceduralFoliageEditorLibrary.generated.h"
+
+class AProceduralFoliageVolume;
+class UProceduralFoliageComponent;
+
+UCLASS()
+class UProceduralFoliageEditorLibrary : public UBlueprintFunctionLibrary
+{
+	GENERATED_UCLASS_BODY()
+
+public:
+	UFUNCTION(BlueprintCallable, Category="Foliage")
+	static void ResimulateProceduralFoliageVolumes(const TArray<AProceduralFoliageVolume*>& ProceduralFoliageVolumes);
+
+	UFUNCTION(BlueprintCallable, Category="Foliage")
+	static void ResimulateProceduralFoliageComponents(const TArray<UProceduralFoliageComponent*>& ProceduralFoliageComponents);
+
+	UFUNCTION(BlueprintCallable, Category="Foliage")
+	static void ClearProceduralFoliageVolumes(const TArray<AProceduralFoliageVolume*>& ProceduralFoliageVolumes);
+
+	UFUNCTION(BlueprintCallable, Category = "Foliage")
+	static void ClearProceduralFoliageComponents(const TArray<UProceduralFoliageComponent*>& ProceduralFoliageComponents);
+};
