Procedural Foliave Volumes
- Don't allow simulation of PFVs when they're covering unloaded areas in World Partition.

diff --git a/Engine/Source/Editor/FoliageEdit/Private/ProceduralFoliageComponentDetails.cpp b/Engine/Source/Editor/FoliageEdit/Private/ProceduralFoliageComponentDetails.cpp
--- a/Engine/Source/Editor/FoliageEdit/Private/ProceduralFoliageComponentDetails.cpp
+++ b/Engine/Source/Editor/FoliageEdit/Private/ProceduralFoliageComponentDetails.cpp
@@ -13,6 +13,7 @@
 #include "ProceduralFoliageComponent.h"
 #include "ProceduralFoliageEditorLibrary.h"
 #include "ScopedTransaction.h"
+#include "WorldPartition/WorldPartition.h"
 
 #define LOCTEXT_NAMESPACE "ProceduralFoliageComponentDetails"
 
@@ -28,6 +29,7 @@ void FProceduralFoliageComponentDetails::CustomizeDetails(IDetailLayoutBuilder&
 	IDetailCategoryBuilder& ProceduralFoliageCategory = DetailBuilder.EditCategory(ProceduralFoliageCategoryName);
 
 	const FText ResimulateText = LOCTEXT("ResimulateButtonText", "Resimulate");
+	const FText LoadUnloadedAreasText = LOCTEXT("LoadUnloadedAreasButtonText", "Load Unloaded Areas");
 
 	TArray< TWeakObjectPtr<UObject> > ObjectsBeingCustomized;
 	DetailBuilder.GetObjectsBeingCustomized(ObjectsBeingCustomized);
@@ -51,19 +53,39 @@ void FProceduralFoliageComponentDetails::CustomizeDetails(IDetailLayoutBuilder&
 		ProceduralFoliageCategory.AddProperty(Property);
 	}
 
-	FDetailWidgetRow& NewRow = ProceduralFoliageCategory.AddCustomRow(ResimulateText);
+	FDetailWidgetRow& NewRow = ProceduralFoliageCategory.AddCustomRow(FText::GetEmpty());
 
 	NewRow.ValueContent()
 		.MaxDesiredWidth(120.f)
 		[
-			SNew(SButton)
-			.OnClicked(this, &FProceduralFoliageComponentDetails::OnResimulateClicked)
-			.ToolTipText(this, &FProceduralFoliageComponentDetails::GetResimulateTooltipText)
-			.IsEnabled(this, &FProceduralFoliageComponentDetails::IsResimulateEnabled)
+			SNew(SHorizontalBox)
+			+ SHorizontalBox::Slot()
+			.AutoWidth()			
+			.Padding(4.0f)
 			[
-				SNew(STextBlock)
-				.Font(IDetailLayoutBuilder::GetDetailFont())
-				.Text(ResimulateText)
+				SNew(SButton)
+				.OnClicked(this, &FProceduralFoliageComponentDetails::OnResimulateClicked)
+				.ToolTipText(this, &FProceduralFoliageComponentDetails::GetResimulateTooltipText)
+				.IsEnabled(this, &FProceduralFoliageComponentDetails::IsResimulateEnabled)
+				[
+					SNew(STextBlock)
+					.Font(IDetailLayoutBuilder::GetDetailFont())
+					.Text(ResimulateText)
+				]
+			]
+			+ SHorizontalBox::Slot()
+			.AutoWidth()
+			.Padding(4.0f)
+			[
+				SNew(SButton)
+				.OnClicked(this, &FProceduralFoliageComponentDetails::OnLoadUnloadedAreas)
+				.ToolTipText(LOCTEXT("Load_UnloadedAreas", "Load unloaded areas required to simulate."))
+				.IsEnabled(this, &FProceduralFoliageComponentDetails::HasUnloadedAreas)
+				[
+					SNew(STextBlock)
+					.Font(IDetailLayoutBuilder::GetDetailFont())
+					.Text(LoadUnloadedAreasText)
+				]
 			]
 		];
 }
@@ -85,72 +107,104 @@ FReply FProceduralFoliageComponentDetails::OnResimulateClicked()
 	return FReply::Handled();
 }
 
-bool FProceduralFoliageComponentDetails::IsResimulateEnabled() const
+FReply FProceduralFoliageComponentDetails::OnLoadUnloadedAreas()
 {
-	bool bCanSimulate = false;
-
 	for(const TWeakObjectPtr<UProceduralFoliageComponent>& Component : SelectedComponents)
 	{
-		if(Component.IsValid() && Component->FoliageSpawner)
+		if(Component.IsValid())
 		{
-			for (const FFoliageTypeObject& FoliageTypeObject : Component->FoliageSpawner->GetFoliageTypes())
+			if (UWorldPartition* WorldPartition = Component->GetWorld()->GetWorldPartition())
 			{
-				// Make sure at least one foliage type is ready to spawn
-				if (FoliageTypeObject.HasFoliageType())
-				{
-					bCanSimulate = true;
-					break;
-				}
-			}
-
-			if (bCanSimulate)
-			{
-				break;
+				FVector Origin;
+				FVector Extent;
+				Component->GetOwner()->GetActorBounds(false, Origin, Extent);
+				WorldPartition->LoadEditorCells(FBox(Origin - Extent, Origin + Extent));
 			}
 		}
 	}
 
-	return bCanSimulate;
+	return FReply::Handled();
+}
+
+bool FProceduralFoliageComponentDetails::IsResimulateEnabled() const
+{
+	FText Reason;
+	return IsResimulateEnabledWithReason(Reason);
 }
 
-FText FProceduralFoliageComponentDetails::GetResimulateTooltipText() const
+bool FProceduralFoliageComponentDetails::IsResimulateEnabledWithReason(FText& OutReason) const
 {
-	FText TooltipText;
+	bool bCanSimulate = false;
 
-	for (const TWeakObjectPtr<UProceduralFoliageComponent>& Component : SelectedComponents)
+	for(const TWeakObjectPtr<UProceduralFoliageComponent>& Component : SelectedComponents)
 	{
-		if (Component.IsValid())
+		if(Component.IsValid())
 		{
 			if (!Component->FoliageSpawner)
 			{
-				TooltipText = LOCTEXT("Resimulate_Tooltip_NeedSpawner", "Cannot generate foliage: Assign a Procedural Foliage Spawner to run the procedural foliage simulation");
+				OutReason = LOCTEXT("Resimulate_Tooltip_NeedSpawner", "Cannot generate foliage: Assign a Procedural Foliage Spawner to run the procedural foliage simulation");
+				return false;
 			}
-			else
+
+			if (!bCanSimulate)
 			{
 				for (const FFoliageTypeObject& FoliageTypeObject : Component->FoliageSpawner->GetFoliageTypes())
 				{
 					// Make sure at least one foliage type is ready to spawn
-					if (!FoliageTypeObject.HasFoliageType())
+					if (FoliageTypeObject.HasFoliageType())
 					{
-						TooltipText = LOCTEXT("Resimulate_Tooltip_EmptySpawner", "Cannot generate foliage: The assigned Procedural Foliage Spawner does not contain any foliage types to spawn.");
+						bCanSimulate = true;
 						break;
 					}
 				}
-			}
 
-			if (!TooltipText.IsEmpty())
-			{
-				break;
+				if (!bCanSimulate)
+				{
+					OutReason = LOCTEXT("Resimulate_Tooltip_EmptySpawner", "Cannot generate foliage: The assigned Procedural Foliage Spawner does not contain any foliage types to spawn.");
+					return false;
+				}
 			}
 		}
 	}
 
-	if (TooltipText.IsEmpty())
+	if (bCanSimulate && HasUnloadedAreas())
 	{
-		TooltipText = LOCTEXT("Resimulate_Tooltip", "Runs the procedural foliage spawner simulation. Replaces any existing instances spawned by a previous simulation.");
+		OutReason = LOCTEXT("Resimulate_Tooltip_UnloadedRegion", "Cannot generate foliage: The assigned Procedural Foliage Volume covers an unloaded area.");
+		return false;
 	}
 
+	OutReason = LOCTEXT("Resimulate_Tooltip", "Runs the procedural foliage spawner simulation. Replaces any existing instances spawned by a previous simulation.");
+	return true;
+}
+
+FText FProceduralFoliageComponentDetails::GetResimulateTooltipText() const
+{
+	FText TooltipText;
+	IsResimulateEnabledWithReason(TooltipText);
 	return TooltipText;
 }
 
+bool FProceduralFoliageComponentDetails::HasUnloadedAreas() const
+{
+	for(const TWeakObjectPtr<UProceduralFoliageComponent>& Component : SelectedComponents)
+	{
+		if(Component.IsValid())
+		{
+			if (UWorldPartition* WorldPartition = Component->GetWorld()->GetWorldPartition())
+			{
+				FVector Origin;
+				FVector Extent;
+				Component->GetOwner()->GetActorBounds(false, Origin, Extent);
+	
+				if (!WorldPartition->AreEditorCellsLoaded(FBox(Origin - Extent, Origin + Extent)))
+				{
+					return true;
+				}
+			}
+		}
+	}
+
+	return false;
+}
+
 #undef LOCTEXT_NAMESPACE
diff --git a/Engine/Source/Editor/FoliageEdit/Private/ProceduralFoliageComponentDetails.h b/Engine/Source/Editor/FoliageEdit/Private/ProceduralFoliageComponentDetails.h
--- a/Engine/Source/Editor/FoliageEdit/Private/ProceduralFoliageComponentDetails.h
+++ b/Engine/Source/Editor/FoliageEdit/Private/ProceduralFoliageComponentDetails.h
@@ -22,8 +22,11 @@ public:
 	virtual void CustomizeDetails( IDetailLayoutBuilder& DetailBuilder );
 private:
 	FReply OnResimulateClicked();
+	FReply OnLoadUnloadedAreas();
 	bool IsResimulateEnabled() const;
+	bool IsResimulateEnabledWithReason(FText& OutReason) const;
 	FText GetResimulateTooltipText() const;
+	bool HasUnloadedAreas() const;
 private:
 	TArray< TWeakObjectPtr<class UProceduralFoliageComponent> > SelectedComponents;
 
diff --git a/Engine/Source/Runtime/Engine/Private/WorldPartition/WorldPartition.cpp b/Engine/Source/Runtime/Engine/Private/WorldPartition/WorldPartition.cpp
--- a/Engine/Source/Runtime/Engine/Private/WorldPartition/WorldPartition.cpp
+++ b/Engine/Source/Runtime/Engine/Private/WorldPartition/WorldPartition.cpp
@@ -470,6 +470,23 @@ void UWorldPartition::UnloadEditorCells(const FBox& Box)
 	UpdateEditorCells(GetCellsToUnload, /*bIsCellShouldBeLoaded*/false);
 }
 
+bool UWorldPartition::AreEditorCellsLoaded(const FBox& Box)
+{
+	TArray<UWorldPartitionEditorCell*> CellsToLoad;
+	if (EditorHash->GetIntersectingCells(Box, CellsToLoad))
+	{
+		for (UWorldPartitionEditorCell* Cell : CellsToLoad)
+		{
+			if (!Cell->bLoaded)
+			{
+				return false;
+			}
+		}
+	}
+
+	return true;
+}
+
 bool UWorldPartition::UpdateEditorCells(TFunctionRef<bool(TArray<UWorldPartitionEditorCell*>&)> GetCellsToProcess, bool bIsCellShouldBeLoaded)
 {
 	FWorldPartionCellUpdateContext CellUpdateContext(this);
diff --git a/Engine/Source/Runtime/Engine/Public/WorldPartition/WorldPartition.h b/Engine/Source/Runtime/Engine/Public/WorldPartition/WorldPartition.h
--- a/Engine/Source/Runtime/Engine/Public/WorldPartition/WorldPartition.h
+++ b/Engine/Source/Runtime/Engine/Public/WorldPartition/WorldPartition.h
@@ -115,6 +115,7 @@ public:
 
 	void LoadEditorCells(const FBox& Box);
 	void UnloadEditorCells(const FBox& Box);
+	bool AreEditorCellsLoaded(const FBox& Box);
 	bool RefreshLoadedEditorCells();
 
 	// PIE Methods
