- ProceduralFoliageVolume: Crash Fix related to indexing in new AISMPartitionActor
- Fix Painting Crash with AverageNormal option
- Avoid Copy of BodyInstance properties on registered components

diff --git a/Engine/Source/Editor/AssetPlacementEdMode/Private/Tools/PlacementBrushToolBase.cpp b/Engine/Source/Editor/AssetPlacementEdMode/Private/Tools/PlacementBrushToolBase.cpp
--- a/Engine/Source/Editor/AssetPlacementEdMode/Private/Tools/PlacementBrushToolBase.cpp
+++ b/Engine/Source/Editor/AssetPlacementEdMode/Private/Tools/PlacementBrushToolBase.cpp
@@ -80,7 +80,7 @@ bool UPlacementBrushToolBase::FindHitResultWithStartAndEndTraceVectors(FHitResul
 
 		return false; };
 
-	return AInstancedFoliageActor::FoliageTrace(EditingWorld, OutHit, FDesiredFoliageInstance(TraceStart, TraceEnd, TraceRadius), NAME_PlacementBrushTool, /* bReturnFaceIndex */ false, FilterFunc);
+	return AInstancedFoliageActor::FoliageTrace(EditingWorld, OutHit, FDesiredFoliageInstance(TraceStart, TraceEnd, /* FoliageType= */ nullptr, TraceRadius), NAME_PlacementBrushTool, /* bReturnFaceIndex */ false, FilterFunc);
 }
 
 FTransform UPlacementBrushToolBase::GetFinalTransformFromHitLocationAndNormal(const FVector& InLocation, const FVector& InNormal)
diff --git a/Engine/Source/Editor/FoliageEdit/Private/FoliageEdMode.cpp b/Engine/Source/Editor/FoliageEdit/Private/FoliageEdMode.cpp
--- a/Engine/Source/Editor/FoliageEdit/Private/FoliageEdMode.cpp
+++ b/Engine/Source/Editor/FoliageEdit/Private/FoliageEdMode.cpp
@@ -861,7 +861,7 @@ void FEdModeFoliage::FoliageBrushTrace(FEditorViewportClient* ViewportClient, co
 			static FName NAME_FoliageBrush = FName(TEXT("FoliageBrush"));
 			FFoliagePaintingGeometryFilter FilterFunc = FFoliagePaintingGeometryFilter(UISettings);
 
-			if (AInstancedFoliageActor::FoliageTrace(World, Hit, FDesiredFoliageInstance(TraceStart, TraceEnd), NAME_FoliageBrush, /* bReturnFaceIndex */ false, FilterFunc))
+			if (AInstancedFoliageActor::FoliageTrace(World, Hit, FDesiredFoliageInstance(TraceStart, TraceEnd, /* FoliageType= */ nullptr), NAME_FoliageBrush, /* bReturnFaceIndex */ false, FilterFunc))
 			{
 				UPrimitiveComponent* PrimComp = Hit.Component.Get();
 				if (PrimComp != nullptr && CanPaint(PrimComp->GetComponentLevel()))
@@ -1508,7 +1508,7 @@ bool FEdModeFoliage::AddSingleInstanceForBrush(UWorld* InWorld, const UFoliageTy
 	FVector Start = BrushLocation + BrushNormal;
 	FVector End = BrushLocation - BrushNormal;
 
-	FDesiredFoliageInstance* DesiredInstance = new (DesiredInstances)FDesiredFoliageInstance(Start, End);
+	FDesiredFoliageInstance* DesiredInstance = new (DesiredInstances)FDesiredFoliageInstance(Start, End, Settings);
 
 	// We do not apply the density limitation based on the brush size
 	TArray<int32> ExistingInstanceBuckets;
@@ -1570,7 +1570,7 @@ void FEdModeFoliage::AddInstancesForBrush(UWorld* InWorld, const UFoliageType* S
 		{
 			FVector Start, End;
 			GetRandomVectorInBrush(Start, End);
-			FDesiredFoliageInstance* DesiredInstance = new (DesiredInstances)FDesiredFoliageInstance(Start, End);
+			FDesiredFoliageInstance* DesiredInstance = new (DesiredInstances)FDesiredFoliageInstance(Start, End, Settings);
 		}
 
 		AddInstancesImp(InWorld, Settings, DesiredInstances, ExistingInstanceBuckets, Pressure, &LandscapeLayerCaches, &UISettings, nullptr, false);
@@ -2388,7 +2388,7 @@ void FEdModeFoliage::ReapplyInstancesForBrush(UWorld* InWorld, AInstancedFoliage
 				FVector ZAxis = Instance.Rotation.Quaternion().GetAxisZ();
 				FVector Start = Instance.Location + 16.f * ZAxis;
 				FVector End = Instance.Location - 16.f * ZAxis;
-				if (AInstancedFoliageActor::FoliageTrace(InWorld, Hit, FDesiredFoliageInstance(Start, End), NAME_ReapplyInstancesForBrush, /* bReturnFaceIndex */ true, FFoliageTraceFilterFunc(), /*bAverageNormal*/ true))
+				if (AInstancedFoliageActor::FoliageTrace(InWorld, Hit, FDesiredFoliageInstance(Start, End, Settings), NAME_ReapplyInstancesForBrush, /* bReturnFaceIndex */ true, FFoliageTraceFilterFunc(), /*bAverageNormal*/ true))
 				{
 					// Reapply the normal
 					if (bReapplyNormal)
@@ -2506,7 +2506,7 @@ void FEdModeFoliage::ReapplyInstancesForBrush(UWorld* InWorld, AInstancedFoliage
 				static const FName NAME_ReapplyInstancesForBrush = TEXT("ReapplyCollisionWithWorld");
 				FVector Start = Instance.Location + FVector(0.f, 0.f, 16.f);
 				FVector End = Instance.Location - FVector(0.f, 0.f, 16.f);
-				if (AInstancedFoliageActor::FoliageTrace(InWorld, Hit, FDesiredFoliageInstance(Start, End), NAME_ReapplyInstancesForBrush))
+				if (AInstancedFoliageActor::FoliageTrace(InWorld, Hit, FDesiredFoliageInstance(Start, End, Settings), NAME_ReapplyInstancesForBrush))
 				{
 					if (!AInstancedFoliageActor::CheckCollisionWithWorld(InWorld, Settings, Instance, Hit.Normal, Hit.Location, Hit.Component.Get()))
 					{
@@ -3049,7 +3049,7 @@ void FEdModeFoliage::SnapSelectedInstancesToGround(UWorld* InWorld)
 
 					for (int32 InstanceIndex : SelectedIndices)
 					{
-						bMovedInstance |= SnapInstanceToGround(IFA, FoliageType->AlignMaxAngle, FoliageInfo, InstanceIndex);
+						bMovedInstance |= SnapInstanceToGround(IFA, FoliageType, FoliageInfo, InstanceIndex);
 					}
 
 					FoliageInfo.PostMoveInstances(SelectedIndices);
@@ -3086,7 +3086,7 @@ void FEdModeFoliage::HandleOnFoliageTypeMeshChanged(UFoliageType* FoliageType)
 	}
 }
 
-bool FEdModeFoliage::SnapInstanceToGround(AInstancedFoliageActor* InIFA, float AlignMaxAngle, FFoliageInfo& Mesh, int32 InstanceIdx)
+bool FEdModeFoliage::SnapInstanceToGround(AInstancedFoliageActor* InIFA, const UFoliageType* Settings, FFoliageInfo& Mesh, int32 InstanceIdx)
 {
 	FFoliageInstance& Instance = Mesh.Instances[InstanceIdx];
 	FVector Start = Instance.Location;
@@ -3094,7 +3094,7 @@ bool FEdModeFoliage::SnapInstanceToGround(AInstancedFoliageActor* InIFA, float A
 
 	FHitResult Hit;
 	static FName NAME_FoliageSnap = FName("FoliageSnap");
-	if (AInstancedFoliageActor::FoliageTrace(InIFA->GetWorld(), Hit, FDesiredFoliageInstance(Start, End), NAME_FoliageSnap, /* bReturnFaceIndex */ false, FFoliageTraceFilterFunc(), (Instance.Flags & FOLIAGE_AlignToNormal)))
+	if (AInstancedFoliageActor::FoliageTrace(InIFA->GetWorld(), Hit, FDesiredFoliageInstance(Start, End, Settings), NAME_FoliageSnap, /* bReturnFaceIndex */ false, FFoliageTraceFilterFunc(), (Instance.Flags & FOLIAGE_AlignToNormal)))
 	{
 		UPrimitiveComponent* HitComponent = Hit.Component.Get();
 
@@ -3137,7 +3137,7 @@ bool FEdModeFoliage::SnapInstanceToGround(AInstancedFoliageActor* InIFA, float A
 		{
 			// Remove previous alignment and align to new normal.
 			Instance.Rotation = Instance.PreAlignRotation;
-			Instance.AlignToNormal(Hit.Normal, AlignMaxAngle);
+			Instance.AlignToNormal(Hit.Normal, Settings->AlignMaxAngle);
 		}
 
 		return true;
diff --git a/Engine/Source/Editor/FoliageEdit/Private/FoliageEdMode.h b/Engine/Source/Editor/FoliageEdit/Private/FoliageEdMode.h
--- a/Engine/Source/Editor/FoliageEdit/Private/FoliageEdMode.h
+++ b/Engine/Source/Editor/FoliageEdit/Private/FoliageEdMode.h
@@ -643,7 +643,7 @@ private:
 	void UpdateWidgetLocationToInstanceSelection();
 
 	/** Snap instance to the ground   */
-	bool SnapInstanceToGround(AInstancedFoliageActor* InIFA, float AlignMaxAngle, FFoliageInfo& Mesh, int32 InstanceIdx);
+	bool SnapInstanceToGround(AInstancedFoliageActor* InIFA, const UFoliageType* Settings, FFoliageInfo& Mesh, int32 InstanceIdx);
 	void SnapSelectedInstancesToGround(UWorld* InWorld);
 
 	/** Callback for when an actor is spawned (to check if it's a new IFA) */
diff --git a/Engine/Source/Runtime/Engine/Private/ISMPartition/ISMComponentDescriptor.cpp b/Engine/Source/Runtime/Engine/Private/ISMPartition/ISMComponentDescriptor.cpp
--- a/Engine/Source/Runtime/Engine/Private/ISMPartition/ISMComponentDescriptor.cpp
+++ b/Engine/Source/Runtime/Engine/Private/ISMPartition/ISMComponentDescriptor.cpp
@@ -1,4 +1,5 @@
 // Copyright Epic Games, Inc. All Rights Reserved.
+
 #include "ISMPartition/ISMComponentDescriptor.h"
 
 #if WITH_EDITOR
@@ -21,7 +22,7 @@ FISMComponentDescriptor FISMComponentDescriptor::CreateFrom(const TSubclassOf<US
 	return ComponentDescriptor;
 }
 
-void FISMComponentDescriptor::InitFrom(UStaticMeshComponent* Template)
+void FISMComponentDescriptor::InitFrom(UStaticMeshComponent* Template, bool bInitBodyInstance)
 {
 	bEnableDiscardOnLoad = false;
 	ComponentClass = Template->GetClass();
@@ -34,7 +35,6 @@ void FISMComponentDescriptor::InitFrom(UStaticMeshComponent* Template)
 	LightingChannels = Template->LightingChannels;
 	bHasCustomNavigableGeometry = Template->bHasCustomNavigableGeometry;
 	CustomDepthStencilWriteMask = Template->CustomDepthStencilWriteMask;
-	BodyInstance.CopyBodyInstancePropertiesFrom(&Template->BodyInstance);
 	VirtualTextureCullMips = Template->VirtualTextureCullMips;
 	TranslucencySortPriority = Template->TranslucencySortPriority;
 	OverriddenLightMapRes = Template->OverriddenLightMapRes;
@@ -63,6 +63,11 @@ void FISMComponentDescriptor::InitFrom(UStaticMeshComponent* Template)
 			bEnableDensityScaling = HISMTemplate->bEnableDensityScaling;
 		}
 	}
+
+	if (bInitBodyInstance)
+	{
+		BodyInstance.CopyBodyInstancePropertiesFrom(&Template->BodyInstance);
+	}
 }
 
 bool FISMComponentDescriptor::operator!=(const FISMComponentDescriptor& Other) const
diff --git a/Engine/Source/Runtime/Engine/Private/ISMPartition/ISMPartitionActor.cpp b/Engine/Source/Runtime/Engine/Private/ISMPartition/ISMPartitionActor.cpp
--- a/Engine/Source/Runtime/Engine/Private/ISMPartition/ISMPartitionActor.cpp
+++ b/Engine/Source/Runtime/Engine/Private/ISMPartition/ISMPartitionActor.cpp
@@ -499,7 +499,7 @@ bool AISMPartitionActor::IsISMComponent(const UPrimitiveComponent* Component) co
 	return false;
 }
 
-int32 AISMPartitionActor::GetISMInstanceIndex(const FISMClientHandle& Handle, UInstancedStaticMeshComponent* ISMComponent, int32 ComponentIndex) const
+int32 AISMPartitionActor::GetISMInstanceIndex(const FISMClientHandle& Handle, const UInstancedStaticMeshComponent* ISMComponent, int32 ComponentIndex) const
 {
 	check(Handle.Guid == Clients[Handle.Index]);
 	for (const FISMComponentData& ComponentData : DescriptorComponents)
diff --git a/Engine/Source/Runtime/Engine/Public/ISMPartition/ISMComponentDescriptor.h b/Engine/Source/Runtime/Engine/Public/ISMPartition/ISMComponentDescriptor.h
--- a/Engine/Source/Runtime/Engine/Public/ISMPartition/ISMComponentDescriptor.h
+++ b/Engine/Source/Runtime/Engine/Public/ISMPartition/ISMComponentDescriptor.h
@@ -21,7 +21,7 @@ struct ENGINE_API FISMComponentDescriptor
 
 #if WITH_EDITOR
 	static FISMComponentDescriptor CreateFrom(const TSubclassOf<UStaticMeshComponent>& ComponentClass);
-	void InitFrom(UStaticMeshComponent* Component);
+	void InitFrom(UStaticMeshComponent* Component, bool bInitBodyInstance = true);
 
 	uint32 ComputeHash() const;
 	UInstancedStaticMeshComponent* CreateComponent(UObject* Outer, FName Name = NAME_None, EObjectFlags ObjectFlags = EObjectFlags::RF_NoFlags) const;
diff --git a/Engine/Source/Runtime/Engine/Public/ISMPartition/ISMPartitionActor.h b/Engine/Source/Runtime/Engine/Public/ISMPartition/ISMPartitionActor.h
--- a/Engine/Source/Runtime/Engine/Public/ISMPartition/ISMPartitionActor.h
+++ b/Engine/Source/Runtime/Engine/Public/ISMPartition/ISMPartitionActor.h
@@ -39,7 +39,7 @@ public:
 	void RemoveISMInstances(const FISMClientHandle& Handle);
 	void SelectISMInstances(const FISMClientHandle& Handle, bool bSelect, const TSet<int32>& Indices);
 	void SetISMInstanceTransform(const FISMClientHandle& Handle, int32 InstanceIndex, const FTransform& NewTransform, bool bTeleport, const TSortedMap<int32, TArray<FTransform>>& InstanceDefinition);
-	int32 GetISMInstanceIndex(const FISMClientHandle& Handle, UInstancedStaticMeshComponent* ISMComponent, int32 ComponentIndex) const;
+	int32 GetISMInstanceIndex(const FISMClientHandle& Handle, const UInstancedStaticMeshComponent* ISMComponent, int32 ComponentIndex) const;
 	FBox GetISMInstanceBounds(const FISMClientHandle& Handle, const TSet<int32>& Indices) const;
 	void ReserveISMInstances(const FISMClientHandle& Handle, int32 AddedInstanceCount, const TSortedMap<int32, TArray<FTransform>>& InstanceDefinition);
 
diff --git a/Engine/Source/Runtime/Foliage/Private/FoliageActor.cpp b/Engine/Source/Runtime/Foliage/Private/FoliageActor.cpp
--- a/Engine/Source/Runtime/Foliage/Private/FoliageActor.cpp
+++ b/Engine/Source/Runtime/Foliage/Private/FoliageActor.cpp
@@ -195,9 +195,9 @@ FTransform FFoliageActor::GetInstanceWorldTransform(int32 InstanceIndex) const
 	return FTransform::Identity;
 }
 
-bool FFoliageActor::IsOwnedComponent(const UPrimitiveComponent* Component) const
+bool FFoliageActor::IsOwnedComponent(const UPrimitiveComponent* PrimitiveComponent) const
 {
-	const AActor* Owner = Component->GetOwner();
+	const AActor* Owner = PrimitiveComponent->GetOwner();
 
 	return ActorInstances.Contains(Owner);
 }
@@ -207,9 +207,9 @@ int32 FFoliageActor::FindIndex(const AActor* InActor) const
 	return ActorInstances.IndexOfByKey(InActor);
 }
 
-int32 FFoliageActor::FindIndex(const UPrimitiveComponent* Component) const
+int32 FFoliageActor::GetInstanceIndexFrom(const UPrimitiveComponent* PrimitiveComponent, int32 ComponentIndex) const
 {
-	return FindIndex(Component->GetOwner());
+	return FindIndex(PrimitiveComponent->GetOwner());
 }
 
 void FFoliageActor::Refresh(bool Async, bool Force)
diff --git a/Engine/Source/Runtime/Foliage/Private/FoliageActor.h b/Engine/Source/Runtime/Foliage/Private/FoliageActor.h
--- a/Engine/Source/Runtime/Foliage/Private/FoliageActor.h
+++ b/Engine/Source/Runtime/Foliage/Private/FoliageActor.h
@@ -39,10 +39,10 @@ struct FFoliageActor : public FFoliageImpl
 	virtual void MoveInstance(int32 InstanceIndex, UObject*& OutInstanceImplementation) override;
 	virtual void SetInstanceWorldTransform(int32 InstanceIndex, const FTransform& Transform, bool bTeleport) override;
 	virtual FTransform GetInstanceWorldTransform(int32 InstanceIndex) const override;
-	virtual bool IsOwnedComponent(const UPrimitiveComponent* Component) const override;
+	virtual bool IsOwnedComponent(const UPrimitiveComponent* PrimitiveComponent) const override;
 	int32 FindIndex(const AActor* InActor) const;
-	virtual int32 FindIndex(const UPrimitiveComponent* HitComponent) const override;
-
+	virtual int32 GetInstanceIndexFrom(const UPrimitiveComponent* PrimitiveComponent, int32 ComponentIndex) const override;
+	
 	virtual void SelectAllInstances(bool bSelect) override;
 	virtual void SelectInstance(bool bSelect, int32 Index) override;
 	virtual void SelectInstances(bool bSelect, const TSet<int32>& SelectedIndices) override;
diff --git a/Engine/Source/Runtime/Foliage/Private/FoliageISMActor.cpp b/Engine/Source/Runtime/Foliage/Private/FoliageISMActor.cpp
--- a/Engine/Source/Runtime/Foliage/Private/FoliageISMActor.cpp
+++ b/Engine/Source/Runtime/Foliage/Private/FoliageISMActor.cpp
@@ -96,7 +96,8 @@ void FFoliageISMActor::Initialize(const UFoliageType* FoliageType)
 	for (UStaticMeshComponent* StaticMeshComponent : StaticMeshComponents)
 	{
 		FISMComponentDescriptor Descriptor;
-		Descriptor.InitFrom(StaticMeshComponent);
+		// Avoid initializing the body instance as we are going to do it in the InitDescriptorFromFoliageType and that Copy of BodyInstance on a registered components will fail.
+		Descriptor.InitFrom(StaticMeshComponent, /*bInitBodyInstance*/ false);
 		InitDescriptorFromFoliageType(Descriptor, FoliageTypeActor);
 		Descriptor.ComputeHash();
 
@@ -333,9 +334,16 @@ void FFoliageISMActor::SelectInstances(bool bSelect, const TSet<int32>& Selected
 	GetIFA()->SelectISMInstances(ClientHandle, bSelect, SelectedIndices);
 }
 
-int32 FFoliageISMActor::GetInstanceIndexFrom(UInstancedStaticMeshComponent* HISMComponent, int32 ComponentIndex) const
+int32 FFoliageISMActor::GetInstanceIndexFrom(const UPrimitiveComponent* PrimitiveComponent, int32 ComponentIndex) const
 {
-	return GetIFA()->GetISMInstanceIndex(ClientHandle, HISMComponent, ComponentIndex);
+	if (ComponentIndex != INDEX_NONE)
+	{
+		if (const UInstancedStaticMeshComponent* ISMComponent = Cast<UInstancedStaticMeshComponent>(PrimitiveComponent))
+		{
+			return GetIFA()->GetISMInstanceIndex(ClientHandle, ISMComponent, ComponentIndex);
+		}
+	}
+	return INDEX_NONE;
 }
 
 FBox FFoliageISMActor::GetSelectionBoundingBox(const TSet<int32>& SelectedIndices) const
diff --git a/Engine/Source/Runtime/Foliage/Private/FoliageISMActor.h b/Engine/Source/Runtime/Foliage/Private/FoliageISMActor.h
--- a/Engine/Source/Runtime/Foliage/Private/FoliageISMActor.h
+++ b/Engine/Source/Runtime/Foliage/Private/FoliageISMActor.h
@@ -9,6 +9,7 @@
 
 struct FFoliageInfo;
 class UInstancedStaticMeshComponent;
+class UPrimitiveComponent;
 class UBlueprint;
 class UFoliageType_Actor;
 
@@ -50,7 +51,7 @@ struct FFoliageISMActor : public FFoliageImpl
 	virtual void SelectAllInstances(bool bSelect) override;
 	virtual void SelectInstance(bool bSelect, int32 Index) override;
 	virtual void SelectInstances(bool bSelect, const TSet<int32>& SelectedIndices) override;
-	virtual int32 GetInstanceIndexFrom(UInstancedStaticMeshComponent* HISMComponent, int32 ComponentIndex) const;
+	virtual int32 GetInstanceIndexFrom(const UPrimitiveComponent* Component, int32 ComponentIndex) const;
 	virtual FBox GetSelectionBoundingBox(const TSet<int32>& SelectedIndices) const override;
 	virtual void ApplySelection(bool bApply, const TSet<int32>& SelectedIndices) override;
 	virtual void ClearSelection(const TSet<int32>& SelectedIndices) override;
diff --git a/Engine/Source/Runtime/Foliage/Private/InstancedFoliage.cpp b/Engine/Source/Runtime/Foliage/Private/InstancedFoliage.cpp
--- a/Engine/Source/Runtime/Foliage/Private/InstancedFoliage.cpp
+++ b/Engine/Source/Runtime/Foliage/Private/InstancedFoliage.cpp
@@ -114,12 +114,12 @@ struct FFoliageStaticMesh : public FFoliageImpl
 	virtual void SetInstanceWorldTransform(int32 InstanceIndex, const FTransform& Transform, bool bTeleport) override;
 	virtual FTransform GetInstanceWorldTransform(int32 InstanceIndex) const override;
 	virtual void PostUpdateInstances() override;
-	virtual bool IsOwnedComponent(const UPrimitiveComponent* HitComponent) const override;
+	virtual bool IsOwnedComponent(const UPrimitiveComponent* PrimitiveComponent) const override;
 
 	virtual void SelectAllInstances(bool bSelect) override;
 	virtual void SelectInstance(bool bSelect, int32 Index) override;
 	virtual void SelectInstances(bool bSelect, const TSet<int32>& SelectedIndices) override;
-	virtual int32 GetInstanceIndexFrom(UInstancedStaticMeshComponent* HISMComponent, int32 ComponentIndex) const override;
+	virtual int32 GetInstanceIndexFrom(const UPrimitiveComponent* PrimitiveComponent, int32 ComponentIndex) const override;
 	virtual FBox GetSelectionBoundingBox(const TSet<int32>& SelectedIndices) const override;
 	virtual void ApplySelection(bool bApply, const TSet<int32>& SelectedIndices) override;
 	virtual void ClearSelection(const TSet<int32>& SelectedIndices) override;
@@ -1091,9 +1091,9 @@ void FFoliageStaticMesh::PostUpdateInstances()
 	Component->MarkRenderStateDirty();
 }
 
-bool FFoliageStaticMesh::IsOwnedComponent(const UPrimitiveComponent* HitComponent) const
+bool FFoliageStaticMesh::IsOwnedComponent(const UPrimitiveComponent* PrimitiveComponent) const
 {
-	return Component == HitComponent;
+	return Component == PrimitiveComponent;
 }
 
 void FFoliageStaticMesh::SelectAllInstances(bool bSelect)
@@ -1120,9 +1120,9 @@ void FFoliageStaticMesh::SelectInstances(bool bSelect, const TSet<int32>& Select
 	Component->MarkRenderStateDirty();
 }
 
-int32 FFoliageStaticMesh::GetInstanceIndexFrom(UInstancedStaticMeshComponent* HISMComponent, int32 ComponentIndex) const
+int32 FFoliageStaticMesh::GetInstanceIndexFrom(const UPrimitiveComponent* PrimitiveComponent, int32 ComponentIndex) const
 {
-	if (IsOwnedComponent(HISMComponent))
+	if (IsOwnedComponent(PrimitiveComponent))
 	{
 		return ComponentIndex;
 	}
@@ -4855,21 +4855,12 @@ bool AInstancedFoliageActor::FoliageTrace(const UWorld* InWorld, FHitResult& Out
 			for (auto& Pair : IFA->FoliageInfos)
 			{
 				const FFoliageInfo& Info = *Pair.Value;
-
-				if (Hit.Item != INDEX_NONE && Info.Implementation->IsOwnedComponent(HitComponent))
+				int32 InstanceIndex = Info.Implementation->GetInstanceIndexFrom(HitComponent, Hit.Item);
+				if (InstanceIndex != INDEX_NONE)
 				{
-					OutHit.Component = CastChecked<UPrimitiveComponent>(IFA->InstanceBaseCache.GetInstanceBasePtr(Info.Instances[Hit.Item].BaseId).Get(), ECastCheckedType::NullAllowed);
+					OutHit.Component = CastChecked<UPrimitiveComponent>(IFA->InstanceBaseCache.GetInstanceBasePtr(Info.Instances[InstanceIndex].BaseId).Get(), ECastCheckedType::NullAllowed);
 					break;
 				}
-				else
-				{
-					int32 InstanceIndex = Info.Implementation->FindIndex(HitComponent);
-					if (InstanceIndex != INDEX_NONE)
-					{
-						OutHit.Component = CastChecked<UPrimitiveComponent>(IFA->InstanceBaseCache.GetInstanceBasePtr(Info.Instances[InstanceIndex].BaseId).Get(), ECastCheckedType::NullAllowed);
-						break;
-					}
-				}
 			}
 
 			// The foliage we are snapping on doesn't have a valid base
@@ -4896,7 +4887,7 @@ bool AInstancedFoliageActor::FoliageTrace(const UWorld* InWorld, FHitResult& Out
 			continue;
 		}
 
-		if (bAverageNormal && DesiredInstance.FoliageType->AverageNormal)
+		if (bAverageNormal && DesiredInstance.FoliageType && DesiredInstance.FoliageType->AverageNormal)
 		{
 			int32 PointSeed = FFoliagePlacementUtil::GetRandomSeedForPosition(FVector2D(Hit.Location));
 			FRandomStream LocalRandomStream(PointSeed);
@@ -4972,7 +4963,7 @@ bool AInstancedFoliageActor::CheckCollisionWithWorld(const UWorld* InWorld, cons
 			//::DrawDebugDirectionalArrow(InWorld, SamplePos, SamplePos - NormalVector*WorldRadius, 10.0f, FColor::Red, true, 30.0f);
 
 			FHitResult Hit;
-			if (AInstancedFoliageActor::FoliageTrace(InWorld, Hit, FDesiredFoliageInstance(SamplePos, SamplePos - NormalVector*WorldRadius)))
+			if (AInstancedFoliageActor::FoliageTrace(InWorld, Hit, FDesiredFoliageInstance(SamplePos, SamplePos - NormalVector*WorldRadius, Settings)))
 			{
 				FVector LocalHit = OriginalTransform.InverseTransformPosition(Hit.Location);
 				
diff --git a/Engine/Source/Runtime/Foliage/Private/ProceduralFoliageTile.cpp b/Engine/Source/Runtime/Foliage/Private/ProceduralFoliageTile.cpp
--- a/Engine/Source/Runtime/Foliage/Private/ProceduralFoliageTile.cpp
+++ b/Engine/Source/Runtime/Foliage/Private/ProceduralFoliageTile.cpp
@@ -487,10 +487,9 @@ void UProceduralFoliageTile::ExtractDesiredInstances(TArray<FDesiredFoliageInsta
 			continue;
 		}
 		
-		FDesiredFoliageInstance* DesiredInst = new (OutInstances)FDesiredFoliageInstance(StartRay, EndRay, Instance.GetMaxRadius());
+		FDesiredFoliageInstance* DesiredInst = new (OutInstances)FDesiredFoliageInstance(StartRay, EndRay, Instance.Type, Instance.GetMaxRadius());
 		DesiredInst->Rotation = Instance.Rotation;
 		DesiredInst->ProceduralGuid = ProceduralGuid;
-		DesiredInst->FoliageType = Instance.Type;
 		DesiredInst->Age = Instance.Age;
 		DesiredInst->ProceduralVolumeBodyInstance = VolumeBodyInstance;
 		DesiredInst->PlacementMode = EFoliagePlacementMode::Procedural;
diff --git a/Engine/Source/Runtime/Foliage/Public/InstancedFoliage.h b/Engine/Source/Runtime/Foliage/Public/InstancedFoliage.h
--- a/Engine/Source/Runtime/Foliage/Public/InstancedFoliage.h
+++ b/Engine/Source/Runtime/Foliage/Public/InstancedFoliage.h
@@ -208,13 +208,12 @@ struct FFoliageImpl
 	virtual void PostUpdateInstances() {}
 	virtual void PreMoveInstances(const TArray<int32>& InInstancesMoved) {}
 	virtual void PostMoveInstances(const TArray<int32>& InInstancesMoved, bool bFinished) {}
-	virtual bool IsOwnedComponent(const UPrimitiveComponent* Component) const = 0;
-	virtual int32 FindIndex(const UPrimitiveComponent* HitComponent) const { return INDEX_NONE; }
-
+	virtual bool IsOwnedComponent(const UPrimitiveComponent* PrimitiveComponent) const = 0;
+	
 	virtual void SelectAllInstances(bool bSelect) = 0;
 	virtual void SelectInstance(bool bSelect, int32 Index) = 0;
 	virtual void SelectInstances(bool bSelect, const TSet<int32>& SelectedIndices) = 0;
-	virtual int32 GetInstanceIndexFrom(UInstancedStaticMeshComponent* HISMComponent, int32 ComponentIndex) const { return INDEX_NONE; }
+	virtual int32 GetInstanceIndexFrom(const UPrimitiveComponent* PrimitiveComponent, int32 ComponentIndex) const = 0;
 	virtual FBox GetSelectionBoundingBox(const TSet<int32>& SelectedIndices) const = 0;
 	virtual void ApplySelection(bool bApply, const TSet<int32>& SelectedIndices) = 0;
 	virtual void ClearSelection(const TSet<int32>& SelectedIndices) = 0;
@@ -546,8 +545,8 @@ struct FDesiredFoliageInstance
 
 	}
 
-	FDesiredFoliageInstance(const FVector& InStartTrace, const FVector& InEndTrace, const float InTraceRadius = 0.f)
-		: FoliageType(nullptr)
+	FDesiredFoliageInstance(const FVector& InStartTrace, const FVector& InEndTrace, const UFoliageType* InFoliageType, const float InTraceRadius = 0.f)
+		: FoliageType(InFoliageType)
 		, StartTrace(InStartTrace)
 		, EndTrace(InEndTrace)
 		, Rotation(ForceInit)
