Foliage UI Cleanup

diff --git a/Engine/Source/Editor/EditorStyle/Private/StarshipStyle.cpp b/Engine/Source/Editor/EditorStyle/Private/StarshipStyle.cpp
--- a/Engine/Source/Editor/EditorStyle/Private/StarshipStyle.cpp
+++ b/Engine/Source/Editor/EditorStyle/Private/StarshipStyle.cpp
@@ -1562,6 +1562,26 @@ void FStarshipEditorStyle::FStyle::SetupGeneralStyles()
 			.SetHighlightColor(FLinearColor(1.0f, 1.0f, 1.0f))
 			.SetShadowOffset(FVector2D(1, 1))
 			.SetShadowColorAndOpacity(FLinearColor(0, 0, 0, 0.9f)));
+
+
+		FLinearColor InputA = FStyleColors::Input.GetSpecifiedColor();
+		InputA .A = .60;
+
+		FLinearColor InputB = FStyleColors::Input.GetSpecifiedColor();
+		InputA .A = .70;
+
+		Set("FoliageEditMode.FloatingButton", FButtonStyle()
+			.SetNormal(FSlateRoundedBoxBrush(InputA, 2))
+			.SetHovered(FSlateRoundedBoxBrush(InputB, 2))
+			.SetPressed(FSlateRoundedBoxBrush(InputB, 2))
+			.SetNormalForeground(FStyleColors::Foreground)
+			.SetHoveredForeground(FStyleColors::ForegroundHover)
+			.SetPressedForeground(FStyleColors::ForegroundHover)
+			.SetDisabledForeground(FStyleColors::White25)
+			.SetNormalPadding(FMargin(4))
+			.SetPressedPadding(FMargin(4))
+		 );
+
 	}
 #endif // WITH_EDITOR || (IS_PROGRAM && WITH_UNREAL_DEVELOPER_TOOLS)
 
diff --git a/Engine/Source/Editor/FoliageEdit/Private/FoliagePaletteItem.cpp b/Engine/Source/Editor/FoliageEdit/Private/FoliagePaletteItem.cpp
--- a/Engine/Source/Editor/FoliageEdit/Private/FoliagePaletteItem.cpp
+++ b/Engine/Source/Editor/FoliageEdit/Private/FoliagePaletteItem.cpp
@@ -228,18 +228,19 @@ TSharedRef<SButton> FFoliagePaletteItemModel::CreateSaveAssetButton(TAttribute<E
 {
 	return 
 		SNew(SButton)
-		.ContentPadding(0)
-		.ButtonStyle(FEditorStyle::Get(), "ToggleButton")
+		.ButtonStyle(FAppStyle::Get(), "FoliageEditMode.FloatingButton")
 		.Visibility(InVisibility)
 		.IsEnabled(this, &FFoliagePaletteItemModel::IsSaveEnabled)
 		.OnClicked(this, &FFoliagePaletteItemModel::HandleSaveAsset)
 		.ToolTipText(LOCTEXT("SaveButtonToolTip", "Save foliage asset"))
-		.HAlign(HAlign_Center)
-		.VAlign(VAlign_Center)
+		.HAlign(HAlign_Right)
+		.VAlign(VAlign_Top)
+		.ContentPadding(0)
 		.Content()
 		[
 			SNew(SImage)
-			.Image(FEditorStyle::GetBrush("Level.SaveIcon16x"))
+	        .ColorAndOpacity(FSlateColor::UseForeground())
+			.Image(FAppStyle::Get().GetBrush("Icons.Save"))
 		];
 }
 
@@ -410,7 +411,7 @@ void SFoliagePaletteItemTile::Construct(const FArguments& InArgs, TSharedRef<STa
 
 	STableRow<FFoliageMeshUIInfoPtr>::Construct(
 		STableRow<FFoliageMeshUIInfoPtr>::FArguments()
-		.Style(FEditorStyle::Get(), "ContentBrowser.AssetListView.TableRow")
+		.Style(FEditorStyle::Get(), "ContentBrowser.AssetListView.ColumnListTableRow")
 		.Padding(1.f)
 		.Content()
 		[
@@ -421,7 +422,7 @@ void SFoliagePaletteItemTile::Construct(const FArguments& InArgs, TSharedRef<STa
 			+ SOverlay::Slot()
 			[
 				SNew(SBorder)
-				.Padding(4.f)
+				.Padding(2.f)
 				.BorderImage(FEditorStyle::GetBrush("ContentBrowser.ThumbnailShadow"))
 				.ForegroundColor(FLinearColor::White)
 				.ColorAndOpacity(this, &SFoliagePaletteItemTile::GetTileColorAndOpacity)
@@ -438,8 +439,6 @@ void SFoliagePaletteItemTile::Construct(const FArguments& InArgs, TSharedRef<STa
 			[
 				SNew(SBorder)
 				.BorderImage(FEditorStyle::GetBrush("ContentBrowser.ThumbnailShadow"))
-				.BorderBackgroundColor(FLinearColor::Black)
-				.ForegroundColor(FLinearColor::White)
 				.Padding(3.f)
 				[
 					Model->CreateActivationCheckBox(IsSelectedGetter, CheckBoxVisibility)
@@ -450,7 +449,7 @@ void SFoliagePaletteItemTile::Construct(const FArguments& InArgs, TSharedRef<STa
 			+ SOverlay::Slot()
 			.HAlign(HAlign_Right)
 			.VAlign(VAlign_Top)
-			.Padding(FMargin(3.f))
+			.Padding(FMargin(2.f))
 			[
 				Model->CreateSaveAssetButton(SaveButtonVisibility)
 			]
diff --git a/Engine/Source/Editor/FoliageEdit/Private/SFoliagePalette.cpp b/Engine/Source/Editor/FoliageEdit/Private/SFoliagePalette.cpp
--- a/Engine/Source/Editor/FoliageEdit/Private/SFoliagePalette.cpp
+++ b/Engine/Source/Editor/FoliageEdit/Private/SFoliagePalette.cpp
@@ -19,6 +19,7 @@
 #include "Framework/MultiBox/MultiBoxBuilder.h"
 #include "Widgets/Layout/SScrollBorder.h"
 #include "Widgets/Input/SSlider.h"
+#include "SEditorHeaderButton.h"
 #include "EditorStyleSet.h"
 #include "FoliagePaletteCommands.h"
 #include "AssetThumbnail.h"
@@ -210,74 +211,66 @@ void SFoliagePalette::Construct(const FArguments& InArgs)
 		+ SVerticalBox::Slot()
 		.AutoHeight()
 		.HAlign(HAlign_Fill)
+		.Padding(FMargin(2.f, 8.f))
 		[
-			// Top bar
-			SNew(SBorder)
-			.BorderImage(FEditorStyle::GetBrush("DetailsView.CategoryTop"))
-			.Padding(FMargin(6.f, 2.f))
-			.BorderBackgroundColor(FLinearColor(.6f, .6f, .6f, 1.0f))
+			SNew(SHorizontalBox)
+
+			+ SHorizontalBox::Slot()
+			.HAlign(HAlign_Left)
+			.AutoWidth()
 			[
-				SNew(SHorizontalBox)
+				// +Add Foliage Type button
+				SAssignNew(AddFoliageTypeCombo, SEditorHeaderButton)
+				.Icon(FAppStyle::Get().GetBrush("Icons.Plus"))
+				.Text(LOCTEXT("AddFoliageTypeButtonLabel", "Foliage"))
+				.OnGetMenuContent(this, &SFoliagePalette::GetAddFoliageTypePicker)
+			]
 
-				+ SHorizontalBox::Slot()
-				.HAlign(HAlign_Left)
-				.AutoWidth()
-				[
-					// +Add Foliage Type button
-					SAssignNew(AddFoliageTypeCombo, SComboButton)
-					.ForegroundColor(FLinearColor::White)
-					.ButtonStyle(FEditorStyle::Get(), "FlatButton.Success")
-					.OnGetMenuContent(this, &SFoliagePalette::GetAddFoliageTypePicker)
-					.ContentPadding(FMargin(1.f))
-					.ButtonContent()
-					[
-						SNew(SHorizontalBox)
-						+ SHorizontalBox::Slot()
-						.VAlign(VAlign_Center)
-						.AutoWidth()
-						.Padding(1.f)
-						[
-							SNew(STextBlock)
-							.TextStyle(FEditorStyle::Get(), "FoliageEditMode.AddFoliageType.Text")
-							.Font(FEditorStyle::Get().GetFontStyle("FontAwesome.9"))
-							.Text(FText::FromString(FString(TEXT("\xf067"))) /*fa-plus*/)
-						]
-						+ SHorizontalBox::Slot()
-						.VAlign(VAlign_Center)
-						.Padding(1.f)
-						[
-							SNew(STextBlock)
-							.Text(LOCTEXT("AddFoliageTypeButtonLabel", "Add Foliage Type"))
-							.TextStyle(FEditorStyle::Get(), "FoliageEditMode.AddFoliageType.Text")
-						]
-					]
-				]
+			+ SHorizontalBox::Slot()
+			.HAlign(HAlign_Fill)
+			.VAlign(VAlign_Center)
+			.Padding(6.f, 0.f)
+			[
+				SAssignNew(SearchBoxPtr, SSearchBox)
+				.HintText(LOCTEXT("SearchFoliagePaletteHint", "Search Foliage"))
+				.OnTextChanged(this, &SFoliagePalette::OnSearchTextChanged)
+			]
 
-				+ SHorizontalBox::Slot()
-				.HAlign(HAlign_Fill)
-				.VAlign(VAlign_Center)
-				.Padding(6.f, 0.f)
+			// Show Details  
+			+ SHorizontalBox::Slot()
+			.HAlign(HAlign_Right)
+			.AutoWidth()
+			[
+				SNew(SCheckBox)
+				.ToolTipText(this, &SFoliagePalette::GetShowHideDetailsTooltipText)
+				.Style(FAppStyle::Get(), "ToggleButtonCheckBox")
+				.IsChecked_Lambda([this]() -> ECheckBoxState {
+						return FoliageEditMode->UISettings.GetShowPaletteItemDetails() ? ECheckBoxState::Checked : ECheckBoxState::Unchecked;
+					})
+
+				.OnCheckStateChanged(this, &SFoliagePalette::OnShowHideDetailsClicked)
+				.Content()
 				[
-					SAssignNew(SearchBoxPtr, SSearchBox)
-					.HintText(LOCTEXT("SearchFoliagePaletteHint", "Search Foliage"))
-					.OnTextChanged(this, &SFoliagePalette::OnSearchTextChanged)
+					SNew(SImage)
+            		.ColorAndOpacity(FSlateColor::UseForeground())
+					.Image(FEditorStyle::GetBrush("LevelEditor.Tabs.Details"))
 				]
+			]
 
-				// View Options
-				+ SHorizontalBox::Slot()
-				.HAlign(HAlign_Right)
-				.AutoWidth()
+			// View Options
+			+ SHorizontalBox::Slot()
+			.HAlign(HAlign_Right)
+			.AutoWidth()
+			[
+				SNew( SComboButton )
+				.ComboButtonStyle( FAppStyle::Get(), "SimpleComboButton" ) // Use the tool bar item style for this button
+				.OnGetMenuContent(this, &SFoliagePalette::GetViewOptionsMenuContent)
+				.HasDownArrow(false)
+				.ButtonContent()
 				[
-					SNew( SComboButton )
-					.ContentPadding(0)
-					.ForegroundColor( FSlateColor::UseForeground() )
-					.ButtonStyle( FEditorStyle::Get(), "ToggleButton" )
-					.OnGetMenuContent(this, &SFoliagePalette::GetViewOptionsMenuContent)
-					.ButtonContent()
-					[
-						SNew(SImage)
-						.Image( FEditorStyle::GetBrush("GenericViewButton") )
-					]
+					SNew(SImage)
+            		.ColorAndOpacity(FSlateColor::UseForeground())
+					.Image( FAppStyle::Get().GetBrush("Icons.Settings") )
 				]
 			]
 		]
@@ -337,44 +330,6 @@ void SFoliagePalette::Construct(const FArguments& InArgs)
 							SNew(STextBlock)
 							.Text(this, &SFoliagePalette::GetDetailsNameAreaText)
 						]
-
-						// Show/Hide Details
-						+ SHorizontalBox::Slot()
-						.HAlign(HAlign_Right)
-						.AutoWidth()
-						[
-							SNew(SButton)
-							.ToolTipText(this, &SFoliagePalette::GetShowHideDetailsTooltipText)
-							.ForegroundColor(FSlateColor::UseForeground())
-							.ButtonStyle(FEditorStyle::Get(), "ToggleButton")
-							.OnClicked(this, &SFoliagePalette::OnShowHideDetailsClicked)
-							.ContentPadding(FMargin(2.f))
-							.Content()
-							[
-								SNew(SHorizontalBox)
-
-								// Details icon
-								+SHorizontalBox::Slot()
-								.AutoWidth()
-								.HAlign(HAlign_Center)
-								.VAlign(VAlign_Center)
-								[
-									SNew(SImage)
-									.Image(FEditorStyle::GetBrush("LevelEditor.Tabs.Details"))
-								]
-
-								// Arrow
-								+ SHorizontalBox::Slot()
-								.Padding(FMargin(3.f, 0.f))
-								.AutoWidth()
-								.HAlign(HAlign_Center)
-								.VAlign(VAlign_Center)
-								[
-									SNew(SImage)
-									.Image(this, &SFoliagePalette::GetShowHideDetailsImage)
-								]
-							]
-						]
 					]
 				]
 				
@@ -568,7 +523,7 @@ void SFoliagePalette::AddFoliageType(const FAssetData& AssetData)
 {
 	if (AddFoliageTypeCombo.IsValid())
 	{
-		AddFoliageTypeCombo->SetIsOpen(false);
+		AddFoliageTypeCombo->SetIsMenuOpen(false, false);
 	}
 
 	GWarn->BeginSlowTask(LOCTEXT("AddFoliageType_LoadPackage", "Loading Foliage Type"), true, false);
@@ -1422,13 +1377,12 @@ const FSlateBrush* SFoliagePalette::GetShowHideDetailsImage() const
 	return FEditorStyle::Get().GetBrush(bDetailsCurrentlyVisible ? "Symbols.DoubleDownArrow" : "Symbols.DoubleUpArrow");
 }
 
-FReply SFoliagePalette::OnShowHideDetailsClicked() const
+void SFoliagePalette::OnShowHideDetailsClicked(const ECheckBoxState InCheckedState) const
 {
-	const bool bDetailsCurrentlyVisible = DetailsWidget->GetVisibility() != EVisibility::Collapsed;
-	DetailsWidget->SetVisibility(bDetailsCurrentlyVisible ? EVisibility::Collapsed : EVisibility::SelfHitTestInvisible);
-	FoliageEditMode->UISettings.SetShowPaletteItemDetails(!bDetailsCurrentlyVisible);
+	const bool bShouldShowDetails = InCheckedState == ECheckBoxState::Checked;
 
-	return FReply::Handled();
+	DetailsWidget->SetVisibility(bShouldShowDetails ? EVisibility::SelfHitTestInvisible : EVisibility::Collapsed );
+	FoliageEditMode->UISettings.SetShowPaletteItemDetails(bShouldShowDetails);
 }
 
 EVisibility SFoliagePalette::GetUneditableFoliageTypeWarningVisibility() const
diff --git a/Engine/Source/Editor/FoliageEdit/Private/SFoliagePalette.h b/Engine/Source/Editor/FoliageEdit/Private/SFoliagePalette.h
--- a/Engine/Source/Editor/FoliageEdit/Private/SFoliagePalette.h
+++ b/Engine/Source/Editor/FoliageEdit/Private/SFoliagePalette.h
@@ -249,7 +249,7 @@ private:	// DETAILS
 	const FSlateBrush* GetShowHideDetailsImage() const;
 
 	/** Handles the show/hide details button click */
-	FReply OnShowHideDetailsClicked() const;
+	void OnShowHideDetailsClicked(const ECheckBoxState ) const;
 
 	/** Gets the visibility of the uneditable blueprint foliage type warning */
 	EVisibility GetUneditableFoliageTypeWarningVisibility() const;
@@ -278,7 +278,7 @@ private:
 	TSharedPtr<class SWidgetSwitcher> WidgetSwitcher;
 
 	/** The Add Foliage Type combo button */
-	TSharedPtr<class SComboButton> AddFoliageTypeCombo;
+	TSharedPtr<class SEditorHeaderButton> AddFoliageTypeCombo;
 
 	/** The header row of the foliage mesh tree */
 	TSharedPtr<class SHeaderRow> TreeViewHeaderRow;
diff --git a/Engine/Source/Editor/UnrealEd/Private/Dialogs/DlgPickAssetPath.cpp b/Engine/Source/Editor/UnrealEd/Private/Dialogs/DlgPickAssetPath.cpp
--- a/Engine/Source/Editor/UnrealEd/Private/Dialogs/DlgPickAssetPath.cpp
+++ b/Engine/Source/Editor/UnrealEd/Private/Dialogs/DlgPickAssetPath.cpp
@@ -84,15 +84,19 @@ void SDlgPickAssetPath::Construct(const FArguments& InArgs)
 			.AutoHeight()
 			.HAlign(HAlign_Right)
 			.VAlign(VAlign_Bottom)
+			.Padding(8.f, 16.f)
 			[
 				SNew(SUniformGridPanel)
-				.SlotPadding(FEditorStyle::GetMargin("StandardDialog.SlotPadding"))
 				.MinDesiredSlotWidth(FEditorStyle::GetFloat("StandardDialog.MinDesiredSlotWidth"))
 				.MinDesiredSlotHeight(FEditorStyle::GetFloat("StandardDialog.MinDesiredSlotHeight"))
+				.SlotPadding(FMargin(8.f, 0.f))
+
 				+SUniformGridPanel::Slot(0,0)
 				[
 					SNew(SButton)
-					.Text(LOCTEXT("OK", "OK"))
+					.ButtonStyle(FAppStyle::Get(), "PrimaryButton")
+					.TextStyle(FAppStyle::Get(), "DialogButtonText")
+					.Text(LOCTEXT("SAVE", "Save"))
 					.HAlign(HAlign_Center)
 					.ContentPadding(FEditorStyle::GetMargin("StandardDialog.ContentPadding"))
 					.OnClicked(this, &SDlgPickAssetPath::OnButtonClick, EAppReturnType::Ok)
@@ -100,6 +104,7 @@ void SDlgPickAssetPath::Construct(const FArguments& InArgs)
 				+SUniformGridPanel::Slot(1,0)
 				[
 					SNew(SButton)
+					.TextStyle(FAppStyle::Get(), "DialogButtonText")
 					.Text(LOCTEXT("Cancel", "Cancel"))
 					.HAlign(HAlign_Center)
 					.ContentPadding(FEditorStyle::GetMargin("StandardDialog.ContentPadding"))
