Fix UpdateInstancePartitioning failed check because PostTransformSelectedInstances needs to get called first to rehash instances

diff --git a/Engine/Source/Editor/FoliageEdit/Private/FoliageEdMode.cpp b/Engine/Source/Editor/FoliageEdit/Private/FoliageEdMode.cpp
--- a/Engine/Source/Editor/FoliageEdit/Private/FoliageEdMode.cpp
+++ b/Engine/Source/Editor/FoliageEdit/Private/FoliageEdMode.cpp
@@ -1862,11 +1862,7 @@ void FEdModeFoliage::ApplySelection(UWorld* InWorld, bool bApply)
 
 void FEdModeFoliage::UpdateInstancePartitioning(UWorld* InWorld)
 {
-	if (!bMoving)
-	{
-		return;
-	}
-
+	check(!bMoving);
 	UActorPartitionSubsystem* ActorPartitionSubsystem = InWorld->GetSubsystem<UActorPartitionSubsystem>();
 	for (TActorIterator<AInstancedFoliageActor> It(InWorld); It; ++It)
 	{
@@ -1938,6 +1934,8 @@ void FEdModeFoliage::PostTransformSelectedInstances(UWorld* InWorld)
 			return true; // continue iteration
 		});
 	}
+
+	UpdateInstancePartitioning(GetWorld());
 }
 
 void FEdModeFoliage::TransformSelectedInstances(UWorld* InWorld, const FVector& InDrag, const FRotator& InRot, const FVector& InScale, bool bDuplicate)
@@ -3881,7 +3879,6 @@ bool FEdModeFoliage::EndTracking()
 	}
 	else if (UISettings.GetSelectToolSelected() || UISettings.GetLassoSelectToolSelected())
 	{
-		UpdateInstancePartitioning(GetWorld());
 		PostTransformSelectedInstances(GetWorld());
 		GEditor->EndTransaction();
 		return true;
