Fix StaticMeshFoliage not being displayed in the Foliage Palette
StaticMeshFoliage can be saved within  a AInstancedFoliageActor of a level instead of saved as an asset.
When in a partionned world, force the creation of an asset so the Foliage can be painted across all cells.
Houdini generation still need to be addressed

diff --git a/Engine/Source/Editor/FoliageEdit/Private/FoliageEdMode.cpp b/Engine/Source/Editor/FoliageEdit/Private/FoliageEdMode.cpp
--- a/Engine/Source/Editor/FoliageEdit/Private/FoliageEdMode.cpp
+++ b/Engine/Source/Editor/FoliageEdit/Private/FoliageEdMode.cpp
@@ -8,6 +8,7 @@
 #include "Components/StaticMeshComponent.h"
 #include "FoliageType.h"
 #include "FoliageType_Actor.h"
+#include "FoliageType_InstancedStaticMesh.h"
 #include "Components/InstancedStaticMeshComponent.h"
 #include "StaticMeshResources.h"
 #include "FoliageInstancedStaticMeshComponent.h"
@@ -3426,24 +3427,29 @@ UFoliageType* FEdModeFoliage::AddFoliageAsset(UObject* InAsset)
 	UStaticMesh* StaticMesh = Cast<UStaticMesh>(InAsset);
 	if (StaticMesh)
 	{
+		AInstancedFoliageActor* IFA = AInstancedFoliageActor::GetDefault(GetWorld());
+		FoliageType = IFA->GetLocalFoliageTypeForSource(StaticMesh);
+		if (!FoliageType)
 		{
-			const FScopedTransaction Transaction(NSLOCTEXT("UnrealEd", "FoliageMode_AddTypeTransaction", "Add Foliage Type"));
-			AInstancedFoliageActor* IFA = AInstancedFoliageActor::GetDefault(GetWorld());
-			FoliageType = IFA->GetLocalFoliageTypeForSource(StaticMesh);
-			if (!FoliageType)
 			{
-				IFA->AddMesh(StaticMesh, &FoliageType);
+				const FScopedTransaction Transaction(NSLOCTEXT("UnrealEd", "FoliageMode_AddTypeTransaction", "Add Foliage Type"));
+				UFoliageType_InstancedStaticMesh *InstancedMeshFoliageType = NewObject<UFoliageType_InstancedStaticMesh>(IFA, NAME_None, RF_Transactional);
+				InstancedMeshFoliageType->SetStaticMesh(StaticMesh);
+				FoliageType = InstancedMeshFoliageType;
 			}
-		}
 
-		// If there is multiple levels for this world, save the foliage directly as an asset, so user will be able to paint over all levels by default
-		if (GetWorld()->GetStreamingLevels().Num() > 0)
-		{
-			UFoliageType* TypeSaved = SaveFoliageTypeObject(FoliageType);
+			if (GetWorld()->GetWorldPartition()
+				|| GetWorld()->GetStreamingLevels().Num() > 0)
+			{
+				// If the world is partitioned or the world has sublevels,
+				// require the user to save the new foliage as an asset.
+				// The user will then be able to paint over all cells/sub-level.
+				FoliageType = SaveFoliageTypeObject(FoliageType);
+			}
 
-			if (TypeSaved != nullptr)
+			if (FoliageType != nullptr)
 			{
-				FoliageType = TypeSaved;
+				IFA->AddMesh(FoliageType);
 			}
 		}
 	}
diff --git a/Engine/Source/Editor/FoliageEdit/Private/FoliageEditUtility.cpp b/Engine/Source/Editor/FoliageEdit/Private/FoliageEditUtility.cpp
--- a/Engine/Source/Editor/FoliageEdit/Private/FoliageEditUtility.cpp
+++ b/Engine/Source/Editor/FoliageEdit/Private/FoliageEditUtility.cpp
@@ -86,27 +86,28 @@ void FFoliageEditUtility::ReplaceFoliageTypeObject(UWorld* InWorld, UFoliageType
 		AInstancedFoliageActor* IFA = *It;
 		IFA->Modify();
 		TUniqueObj<FFoliageInfo> OldInfo;
-		IFA->RemoveFoliageInfoAndCopyValue(OldType, OldInfo);
-
-		// Old component needs to go
-		if (OldInfo->IsInitialized())
+		if (IFA->RemoveFoliageInfoAndCopyValue(OldType, OldInfo))
 		{
-			OldInfo->Uninitialize();
-		}
+			// Old component needs to go
+			if (OldInfo->IsInitialized())
+			{
+				OldInfo->Uninitialize();
+			}
 
-		// Append instances if new foliage type is already exists in this actor
-		// Otherwise just replace key entry for instances
-		FFoliageInfo* NewInfo = IFA->FindInfo(NewType);
-		if (NewInfo)
-		{
-			NewInfo->Instances.Append(OldInfo->Instances);
-			NewInfo->ReallocateClusters(NewType);
-		}
-		else
-		{
-			// Make sure if type changes we have proper implementation
-			TUniqueObj<FFoliageInfo>& NewFoliageInfo = IFA->AddFoliageInfo(NewType, MoveTemp(OldInfo));
-			NewFoliageInfo->ReallocateClusters(NewType);
+			// Append instances if new foliage type is already exists in this actor
+			// Otherwise just replace key entry for instances
+			FFoliageInfo* NewInfo = IFA->FindInfo(NewType);
+			if (NewInfo)
+			{
+				NewInfo->Instances.Append(OldInfo->Instances);
+				NewInfo->ReallocateClusters(NewType);
+			}
+			else
+			{
+				// Make sure if type changes we have proper implementation
+				TUniqueObj<FFoliageInfo>& NewFoliageInfo = IFA->AddFoliageInfo(NewType, MoveTemp(OldInfo));
+				NewFoliageInfo->ReallocateClusters(NewType);
+			}
 		}
 	}
 }
