Remove user-visible entry points to the tutorial system that are still visible when the system is disabled.

diff --git a/Engine/Source/Editor/FoliageEdit/Private/SFoliageEdit.cpp b/Engine/Source/Editor/FoliageEdit/Private/SFoliageEdit.cpp
--- a/Engine/Source/Editor/FoliageEdit/Private/SFoliageEdit.cpp
+++ b/Engine/Source/Editor/FoliageEdit/Private/SFoliageEdit.cpp
@@ -23,7 +23,6 @@
 #include "Classes/EditorStyleSettings.h"
 
 #include "FoliageEditActions.h"
-#include "IIntroTutorials.h"
 #include "Widgets/Input/SNumericEntryBox.h"
 #include "SFoliagePalette.h"
 #include "Widgets/Layout/SHeader.h"
@@ -45,8 +44,6 @@ void SFoliageEdit::Construct(const FArguments& InArgs)
 {
 	FoliageEditMode = (FEdModeFoliage*)GLevelEditorModeTools().GetActiveMode(FBuiltinEditorModes::EM_Foliage);
 
-	IIntroTutorials& IntroTutorials = FModuleManager::LoadModuleChecked<IIntroTutorials>(TEXT("IntroTutorials"));
-
 	// Everything (or almost) uses this padding, change it to expand the padding.
 	FMargin StandardPadding(6.f, 3.f);
 	FMargin StandardLeftPadding(6.f, 3.f, 3.f, 3.f);
@@ -94,16 +91,6 @@ void SFoliageEdit::Construct(const FArguments& InArgs)
 							.Text(this, &SFoliageEdit::GetActiveToolName)
 							.TextStyle(FEditorStyle::Get(), "FoliageEditMode.ActiveToolName.Text")
 						]
-
-						+ SHorizontalBox::Slot()
-						.Padding(StandardRightPadding)
-						.HAlign(HAlign_Right)
-						.VAlign(VAlign_Center)
-						.AutoWidth()
-						[
-							// Tutorial link
-							IntroTutorials.CreateTutorialsWidget(TEXT("FoliageMode"))
-						]
 					]
 
 					+ SVerticalBox::Slot()
diff --git a/Engine/Source/Editor/LandscapeEditor/Private/SLandscapeEditor.cpp b/Engine/Source/Editor/LandscapeEditor/Private/SLandscapeEditor.cpp
--- a/Engine/Source/Editor/LandscapeEditor/Private/SLandscapeEditor.cpp
+++ b/Engine/Source/Editor/LandscapeEditor/Private/SLandscapeEditor.cpp
@@ -16,7 +16,6 @@
 #include "LandscapeEditorObject.h"
 #include "IDetailsView.h"
 #include "PropertyEditorModule.h"
-#include "IIntroTutorials.h"
 
 #define LOCTEXT_NAMESPACE "LandscapeEditor"
 
@@ -450,8 +449,6 @@ void SLandscapeEditor::Construct(const FArguments& InArgs, TSharedRef<FLandscape
 		DetailsPanel->SetObject(LandscapeEdMode->UISettings);
 	}
 
-	IIntroTutorials& IntroTutorials = FModuleManager::LoadModuleChecked<IIntroTutorials>(TEXT("IntroTutorials"));
-
 	ChildSlot
 	[
 		SNew(SVerticalBox)
@@ -467,20 +464,6 @@ void SLandscapeEditor::Construct(const FArguments& InArgs, TSharedRef<FLandscape
 			SNew(SVerticalBox)
 			.IsEnabled(this, &SLandscapeEditor::GetLandscapeEditorIsEnabled)
 			+ SVerticalBox::Slot()
-			.AutoHeight()
-			.Padding(4, 0, 4, 5)
-			[
-				SNew(SOverlay)
-				// Tutorial link
-				+ SOverlay::Slot()
-				.HAlign(HAlign_Right)
-				.VAlign(VAlign_Bottom)
-				.Padding(4)
-				[
-					IntroTutorials.CreateTutorialsWidget(TEXT("LandscapeMode"))
-				]
-			]
-			+ SVerticalBox::Slot()
 			.Padding(0)
 			[
 				DetailsPanel.ToSharedRef()
diff --git a/Engine/Source/Editor/UnrealEd/Private/Toolkits/SStandaloneAssetEditorToolkitHost.cpp b/Engine/Source/Editor/UnrealEd/Private/Toolkits/SStandaloneAssetEditorToolkitHost.cpp
--- a/Engine/Source/Editor/UnrealEd/Private/Toolkits/SStandaloneAssetEditorToolkitHost.cpp
+++ b/Engine/Source/Editor/UnrealEd/Private/Toolkits/SStandaloneAssetEditorToolkitHost.cpp
@@ -14,7 +14,6 @@
 #include "UObject/Package.h"
 #include "StatusBarSubsystem.h"
 #include "ToolMenus.h"
-#include "IIntroTutorials.h"
 #include "EditorModeManager.h"
 #include "EditorModes.h"
 
@@ -76,9 +75,6 @@ void SStandaloneAssetEditorToolkitHost::SetupInitialContent( const TSharedRef<FT
 
 	if (InHostTab)
 	{
-		IIntroTutorials& IntroTutorials = FModuleManager::LoadModuleChecked<IIntroTutorials>(TEXT("IntroTutorials"));
-		TSharedRef<SWidget> TutorialWidget = IntroTutorials.CreateTutorialsWidget(HostedAssetEditorToolkit->GetToolkitContextFName(), InHostTab->GetParentWindow());
-
 		InHostTab->SetRightContent(
 			SNew(SHorizontalBox)
 			+SHorizontalBox::Slot()
@@ -87,13 +83,6 @@ void SStandaloneAssetEditorToolkitHost::SetupInitialContent( const TSharedRef<FT
 			[
 				MenuOverlayWidgetContent.ToSharedRef()
 			]
-			+ SHorizontalBox::Slot()
-			.AutoWidth()
-			.Padding(8.0f, 0.0f, 8.0f, 0.0f)
-			.VAlign(VAlign_Center)
-			[
-				TutorialWidget
-			]
 		);
 	}
 
