Ctrl+Mouse Wheel in Foliage editing mode incorrectly adjusted the size up by 125% per mousewheel click and down by 125% (clamped to 0.1). Now it adjusts in 5% increments, as was intended.

Removed Shift+Ctrl+MMB drag resizing of the brush - it was broken in an equivalent way, and was just super weird anyway, using the Y axis in world space to resize instead of in screen space. I doubt anyone has used it in a long time or the bugs would have been very obvious.

diff --git a/Engine/Source/Editor/FoliageEdit/Private/FoliageEdMode.cpp b/Engine/Source/Editor/FoliageEdit/Private/FoliageEdMode.cpp
--- a/Engine/Source/Editor/FoliageEdit/Private/FoliageEdMode.cpp
+++ b/Engine/Source/Editor/FoliageEdit/Private/FoliageEdMode.cpp
@@ -231,7 +231,6 @@ FEdModeFoliage::FEdModeFoliage()
 	: FEdMode()
 	, bToolActive(false)
 	, bCanAltDrag(false)
-	, bAdjustBrushRadius(false)
 	, FoliageMeshListSortMode(EColumnSortMode::Ascending)
 	, UpdateSelectionCounter(0)
 	, bHasDeferredSelectionNotification(false)
@@ -850,12 +849,9 @@ void FEdModeFoliage::FoliageBrushTrace(FEditorViewportClient* ViewportClient, co
 				UPrimitiveComponent* PrimComp = Hit.Component.Get();
 				if (PrimComp != nullptr && CanPaint(PrimComp->GetComponentLevel()))
 				{
-					if (!bAdjustBrushRadius)
-					{
-						// Adjust the brush location
-						BrushLocation = Hit.Location;
-						BrushNormal = Hit.Normal;
-					}
+					// Adjust the brush location
+					BrushLocation = Hit.Location;
+					BrushNormal = Hit.Normal;
 
 					// Still want to draw the brush when resizing
 					bBrushTraceValid = true;
@@ -3659,16 +3655,15 @@ bool FEdModeFoliage::InputKey(FEditorViewportClient* ViewportClient, FViewport*
 		if (IsCtrlDown(Viewport))
 		{
 			// Control + scroll adjusts the brush radius
-			static const float RadiusAdjustmentAmount = 25.f;
 			if (Key == EKeys::MouseScrollUp)
 			{
-				AdjustBrushRadius(RadiusAdjustmentAmount);
+				AdjustBrushRadius(1);
 
 				bHandled = true;
 			}
 			else if (Key == EKeys::MouseScrollDown)
 			{
-				AdjustBrushRadius(-RadiusAdjustmentAmount);
+				AdjustBrushRadius(-1);
 
 				bHandled = true;
 			}
@@ -3835,15 +3830,12 @@ FVector FEdModeFoliage::GetWidgetLocation() const
 	return FEdMode::GetWidgetLocation();
 }
 
+
 /** FEdMode: Called when a mouse button is pressed */
 bool FEdModeFoliage::StartTracking(FEditorViewportClient* InViewportClient, FViewport* InViewport)
 {
 	bTracking = true;
-	if (IsCtrlDown(InViewport) && InViewport->KeyState(EKeys::MiddleMouseButton) && (UISettings.GetPaintToolSelected() || UISettings.GetReapplyToolSelected() || UISettings.GetLassoSelectToolSelected()))
-	{
-		bAdjustBrushRadius = true;
-	}
-	else if (UISettings.GetSelectToolSelected() || UISettings.GetLassoSelectToolSelected())
+	if (UISettings.GetSelectToolSelected() || UISettings.GetLassoSelectToolSelected())
 	{
 		// Update pivot
 		UpdateWidgetLocationToInstanceSelection();
@@ -3865,12 +3857,7 @@ bool FEdModeFoliage::EndTracking()
 {
 	bTracking = false;
 
-	if (bAdjustBrushRadius)
-	{
-		bAdjustBrushRadius = false;
-		return true;
-	}
-	else if (UISettings.GetSelectToolSelected() || UISettings.GetLassoSelectToolSelected())
+	if (UISettings.GetSelectToolSelected() || UISettings.GetLassoSelectToolSelected())
 	{
 		PostTransformSelectedInstances(GetWorld());
 		GEditor->EndTransaction();
@@ -3894,29 +3881,16 @@ bool FEdModeFoliage::EndTracking(FEditorViewportClient* InViewportClient, FViewp
 /** FEdMode: Called when mouse drag input it applied */
 bool FEdModeFoliage::InputDelta(FEditorViewportClient* InViewportClient, FViewport* InViewport, FVector& InDrag, FRotator& InRot, FVector& InScale)
 {
-	if (bAdjustBrushRadius)
-	{
-		if (UISettings.GetPaintToolSelected() || UISettings.GetReapplyToolSelected() || UISettings.GetLassoSelectToolSelected())
-		{
-			static const float RadiusAdjustmentFactor = 10.f;
-			AdjustBrushRadius(RadiusAdjustmentFactor * InDrag.Y);
-
-			return true;
-		}
-	}
-	else
+	if (InViewportClient->GetCurrentWidgetAxis() != EAxisList::None && (UISettings.GetSelectToolSelected() || UISettings.GetLassoSelectToolSelected()))
 	{
-		if (InViewportClient->GetCurrentWidgetAxis() != EAxisList::None && (UISettings.GetSelectToolSelected() || UISettings.GetLassoSelectToolSelected()))
-		{
-			const bool bDuplicateInstances = (bCanAltDrag && IsAltDown(InViewport) && (InViewportClient->GetCurrentWidgetAxis() & EAxisList::XYZ));
+		const bool bDuplicateInstances = (bCanAltDrag && IsAltDown(InViewport) && (InViewportClient->GetCurrentWidgetAxis() & EAxisList::XYZ));
 
-			TransformSelectedInstances(GetWorld(), InDrag, InRot, InScale, bDuplicateInstances);
+		TransformSelectedInstances(GetWorld(), InDrag, InRot, InScale, bDuplicateInstances);
 
-			// Only allow alt-drag on first InputDelta
-			bCanAltDrag = false;
+		// Only allow alt-drag on first InputDelta
+		bCanAltDrag = false;
 
-			return true;
-		}
+		return true;
 	}
 
 	return FEdMode::InputDelta(InViewportClient, InViewport, InDrag, InRot, InScale);
diff --git a/Engine/Source/Editor/FoliageEdit/Private/FoliageEdMode.h b/Engine/Source/Editor/FoliageEdit/Private/FoliageEdMode.h
--- a/Engine/Source/Editor/FoliageEdit/Private/FoliageEdMode.h
+++ b/Engine/Source/Editor/FoliageEdit/Private/FoliageEdMode.h
@@ -697,7 +697,6 @@ private:
 
 	bool bToolActive;
 	bool bCanAltDrag;
-	bool bAdjustBrushRadius;
 
 	TArray<FFoliageMeshUIInfoPtr>	FoliageMeshList;
 	EColumnSortMode::Type			FoliageMeshListSortMode;
