Remove unused data layer label

diff --git a/Engine/Source/Editor/FoliageEdit/Private/SFoliageEdit.cpp b/Engine/Source/Editor/FoliageEdit/Private/SFoliageEdit.cpp
--- a/Engine/Source/Editor/FoliageEdit/Private/SFoliageEdit.cpp
+++ b/Engine/Source/Editor/FoliageEdit/Private/SFoliageEdit.cpp
@@ -222,25 +222,6 @@ void SFoliageEdit::Construct(const FArguments& InArgs)
 							.IsEnabled(this, &SFoliageEdit::IsEnabled_EraseDensity)
 						]
 					]
-
-					// Data Layer
-					+ SVerticalBox::Slot()
-					.AutoHeight()
-					[
-						SNew(SHorizontalBox)
-						.ToolTipText(LOCTEXT("DataLayer_Tooltip", "The Data Layer to use for the foliage actors"))
-						.Visibility(this, &SFoliageEdit::GetVisibility_DataLayer)
-
-						+ SHorizontalBox::Slot()
-						.Padding(StandardLeftPadding)
-						.FillWidth(1.0f)
-						.VAlign(VAlign_Center)
-						[
-							SNew(STextBlock)
-							.Text(LOCTEXT("DataLayer_Text", "Data Layer"))
-							.Font(StandardFont)
-						]
-					]
 					
 					+ SVerticalBox::Slot()
 					.Padding(StandardPadding)
