Fixing a regression in the Foliage editor, due to a bad combination of splitter vs scrollbox.

When a foliage asset was selected the splitter was requested a lot more space, which was granted starting with CL20995097 (addition of a scrollbox).
The issue was that the handle of the splitter is using a ratio of the splitter total size to position itself, so increasing the size a lot was also moving the handle a lot, assigning far too much space to the foliage asset list view, while the details view was far below.
The fix removes the splitter, and instead piggy back on the vertical box that was already there, and use an autosize/max height slot for the tileview, thus removing the need for a user manual move of splitter's handle.

diff --git a/Engine/Source/Editor/FoliageEdit/Private/SFoliagePalette.cpp b/Engine/Source/Editor/FoliageEdit/Private/SFoliagePalette.cpp
--- a/Engine/Source/Editor/FoliageEdit/Private/SFoliagePalette.cpp
+++ b/Engine/Source/Editor/FoliageEdit/Private/SFoliagePalette.cpp
@@ -297,7 +297,7 @@ void SFoliagePalette::Construct(const FArguments& InArgs)
 				.Content()
 				[
 					SNew(SImage)
-            		.ColorAndOpacity(FSlateColor::UseForeground())
+					.ColorAndOpacity(FSlateColor::UseForeground())
 					.Image(FAppStyle::GetBrush("LevelEditor.Tabs.Details"))
 				]
 			]
@@ -314,109 +314,102 @@ void SFoliagePalette::Construct(const FArguments& InArgs)
 				.ButtonContent()
 				[
 					SNew(SImage)
-            		.ColorAndOpacity(FSlateColor::UseForeground())
+					.ColorAndOpacity(FSlateColor::UseForeground())
 					.Image( FAppStyle::Get().GetBrush("Icons.Settings") )
 				]
 			]
 		]
 
 		+ SVerticalBox::Slot()
+		.AutoHeight()
+		.MaxHeight(230.0f)
 		[
-			SNew(SSplitter)
-			.Orientation(Orient_Vertical)
-			.Style(FAppStyle::Get(), "FoliageEditMode.Splitter")
+			SNew(SOverlay)
 
-			+ SSplitter::Slot()
-			.Value(0.6f)
+			+ SOverlay::Slot()
 			[
-				SNew(SOverlay)
+				SNew(SVerticalBox)
 
-				+ SOverlay::Slot()
+				+ SVerticalBox::Slot()
+				.AutoHeight()
+				.Padding(FMargin(6.f, 3.f))
 				[
-					SNew(SVerticalBox)
-
-					+ SVerticalBox::Slot()
-					.AutoHeight()
-					.Padding(FMargin(6.f, 3.f))
+					SNew(SBox)
+					.Visibility(this, &SFoliagePalette::GetDropFoliageHintVisibility)
+					.Padding(FMargin(15, 0))
+					.MinDesiredHeight(30)
 					[
-						SNew(SBox)
-						.Visibility(this, &SFoliagePalette::GetDropFoliageHintVisibility)
-						.Padding(FMargin(15, 0))
-						.MinDesiredHeight(30)
+						SNew(SScaleBox)
+						.Stretch(EStretch::ScaleToFit)
 						[
-							SNew(SScaleBox)
-							.Stretch(EStretch::ScaleToFit)
-							[
-								SNew(STextBlock)
-								.Text(LOCTEXT("Foliage_DropStatic", "+ Drop Foliage Here"))
-								.ToolTipText(LOCTEXT("Foliage_DropStatic_ToolTip", "Drag and drop foliage types or static meshes from the Content Browser to add them to the palette"))
-							]
+							SNew(STextBlock)
+							.Text(LOCTEXT("Foliage_DropStatic", "+ Drop Foliage Here"))
+							.ToolTipText(LOCTEXT("Foliage_DropStatic_ToolTip", "Drag and drop foliage types or static meshes from the Content Browser to add them to the palette"))
 						]
 					]
+				]
 
-					+ SVerticalBox::Slot()
-					[
-						CreatePaletteViews()
-					]
+				+ SVerticalBox::Slot()
+				[
+					CreatePaletteViews()
+				]
 				
-					+ SVerticalBox::Slot()
-					.Padding(FMargin(0.f))
+				+ SVerticalBox::Slot()
+				.Padding(FMargin(0.f))
+				.VAlign(VAlign_Bottom)
+				.AutoHeight()
+				[
+					SNew(SHorizontalBox)
+						
+					// Selected type name area
+					+ SHorizontalBox::Slot()
+					.Padding(FMargin(3.f))
 					.VAlign(VAlign_Bottom)
-					.AutoHeight()
+					//.AutoWidth()
 					[
-						SNew(SHorizontalBox)
-						
-						// Selected type name area
-						+ SHorizontalBox::Slot()
-						.Padding(FMargin(3.f))
-						.VAlign(VAlign_Bottom)
-						//.AutoWidth()
-						[
-							SNew(STextBlock)
-							.Text(this, &SFoliagePalette::GetDetailsNameAreaText)
-						]
+						SNew(STextBlock)
+						.Text(this, &SFoliagePalette::GetDetailsNameAreaText)
 					]
 				]
+			]
 				
-				// Foliage Mesh Drop Zone
-				+ SOverlay::Slot()
-				.HAlign(HAlign_Fill)
-				.VAlign(VAlign_Fill)
+			// Foliage Mesh Drop Zone
+			+ SOverlay::Slot()
+			.HAlign(HAlign_Fill)
+			.VAlign(VAlign_Fill)
+			[
+				SNew(SFoliageDragDropHandler)
+				.Visibility(this, &SFoliagePalette::GetFoliageDropTargetVisibility)
+				.OnDrop(this, &SFoliagePalette::HandleFoliageDropped)
 				[
-					SNew(SFoliageDragDropHandler)
-					.Visibility(this, &SFoliagePalette::GetFoliageDropTargetVisibility)
-					.OnDrop(this, &SFoliagePalette::HandleFoliageDropped)
+					SNew(SScaleBox)
+					.Stretch(EStretch::ScaleToFit)
 					[
-						SNew(SScaleBox)
-						.Stretch(EStretch::ScaleToFit)
-						[
-							SNew(STextBlock)
-							.Text(LOCTEXT("Foliage_AddFoliageMesh", "+ Foliage Type"))
-							.ShadowOffset(FVector2D(1.f, 1.f))
-						]
+						SNew(STextBlock)
+						.Text(LOCTEXT("Foliage_AddFoliageMesh", "+ Foliage Type"))
+						.ShadowOffset(FVector2D(1.f, 1.f))
 					]
 				]
 			]
+		]
+		// Details
+		+ SVerticalBox::Slot()
+		[
+			SNew(SVerticalBox)
 
-			// Details
-			+SSplitter::Slot()
+			+ SVerticalBox::Slot()
+			.Padding(0, 2)
+			.AutoHeight()
 			[
-				SNew(SVerticalBox)
-
-				+ SVerticalBox::Slot()
-				.Padding(0, 2)
-				.AutoHeight()
-				[
-					SNew(SUneditableFoliageTypeWarning)
-					.WarningText(LOCTEXT("CannotEditBlueprintFoliageTypeWarning", "Blueprint foliage types must be edited in the <a id=\"HyperlinkDecorator\" style=\"DetailsView.BPMessageHyperlinkStyle\">Blueprint</>"))
-					.OnHyperlinkClicked(this, &SFoliagePalette::OnEditFoliageTypeBlueprintHyperlinkClicked)
-					.Visibility(this, &SFoliagePalette::GetUneditableFoliageTypeWarningVisibility)
-				]
+				SNew(SUneditableFoliageTypeWarning)
+				.WarningText(LOCTEXT("CannotEditBlueprintFoliageTypeWarning", "Blueprint foliage types must be edited in the <a id=\"HyperlinkDecorator\" style=\"DetailsView.BPMessageHyperlinkStyle\">Blueprint</>"))
+				.OnHyperlinkClicked(this, &SFoliagePalette::OnEditFoliageTypeBlueprintHyperlinkClicked)
+				.Visibility(this, &SFoliagePalette::GetUneditableFoliageTypeWarningVisibility)
+			]
 
-				+ SVerticalBox::Slot()
-				[
-					DetailsWidget.ToSharedRef()
-				]
+			+ SVerticalBox::Slot()
+			[
+				DetailsWidget.ToSharedRef()
 			]
 		]
 	];
