Foliage Actor type was not available in the Foliage palette even for non-partitioned worlds.

diff --git a/Engine/Source/Editor/FoliageEdit/Private/FoliageEdMode.cpp b/Engine/Source/Editor/FoliageEdit/Private/FoliageEdMode.cpp
--- a/Engine/Source/Editor/FoliageEdit/Private/FoliageEdMode.cpp
+++ b/Engine/Source/Editor/FoliageEdit/Private/FoliageEdMode.cpp
@@ -2067,7 +2067,7 @@ void FEdModeFoliage::RemoveSelectedInstances(UWorld* InWorld)
 void FEdModeFoliage::GetFoliageTypeFilters(TArray<const UClass*>& OutFilters) const
 {
 	OutFilters.Add(UFoliageType_InstancedStaticMesh::StaticClass());
-	if (GetWorld()->HasSubsystem<UWorldPartitionSubsystem>())
+	if (GetWorld()->IsPartitionedWorld())
 	{
 		return;
 	}
