FoliageEdMode: Fix rare crash occurence with zero delta where a Move + Selection was done in the same tracking  (UDN report)

FEditorViewportClient still misses some clicks with very subtle mouse movements

diff --git a/Engine/Source/Editor/FoliageEdit/Private/FoliageEdMode.cpp b/Engine/Source/Editor/FoliageEdit/Private/FoliageEdMode.cpp
--- a/Engine/Source/Editor/FoliageEdit/Private/FoliageEdMode.cpp
+++ b/Engine/Source/Editor/FoliageEdit/Private/FoliageEdMode.cpp
@@ -3786,6 +3786,13 @@ bool FEdModeFoliage::HandleClick(FEditorViewportClient* InViewportClient, HHitPr
 
 	if (UISettings.GetSelectToolSelected())
 	{
+		// If we are moving we can't change the selection. HandleClick can be called in some cases where the input delta is very small.
+		// The condition in FEdModeFoliage::InputDelta where we test that drag is nearly zero should prevent this.
+		if (bMoving)
+		{
+			return true;
+		}
+
 		if (HitProxy && HitProxy->IsA(HInstancedStaticMeshInstance::StaticGetType()))
 		{
 			HInstancedStaticMeshInstance* SMIProxy = ((HInstancedStaticMeshInstance*)HitProxy);
@@ -3901,14 +3908,19 @@ bool FEdModeFoliage::InputDelta(FEditorViewportClient* InViewportClient, FViewpo
 {
 	if (InViewportClient->GetCurrentWidgetAxis() != EAxisList::None && (UISettings.GetSelectToolSelected() || UISettings.GetLassoSelectToolSelected()))
 	{
-		const bool bDuplicateInstances = (bCanAltDrag && IsAltDown(InViewport) && (InViewportClient->GetCurrentWidgetAxis() & EAxisList::XYZ));
+		// It is possible to receive a zero delta from FEditorViewportClient::UpdateMouseDelta which can have a non zero DragDelta which gets converted to zero drag/rot/scale
+		// In which case we want to avoid starting a move in case we then receive a HandleClick in the same frame
+		if (!InDrag.IsNearlyZero() || !InRot.IsNearlyZero() || !InScale.IsNearlyZero())
+		{
+			const bool bDuplicateInstances = (bCanAltDrag && IsAltDown(InViewport) && (InViewportClient->GetCurrentWidgetAxis() & EAxisList::XYZ));
 
-		TransformSelectedInstances(GetWorld(), InDrag, InRot, InScale, bDuplicateInstances);
+			TransformSelectedInstances(GetWorld(), InDrag, InRot, InScale, bDuplicateInstances);
 
-		// Only allow alt-drag on first InputDelta
-		bCanAltDrag = false;
+			// Only allow alt-drag on first InputDelta
+			bCanAltDrag = false;
 
-		return true;
+			return true;
+		}
 	}
 
 	return FEdMode::InputDelta(InViewportClient, InViewport, InDrag, InRot, InScale);
