We are seeing rare crashes in this check that there are no transactions active when exiting foliage edit mode.

No known repro steps, and unclear what circumstances trigger it.

I'm converting it to an ensure, so at least it doesn't crash, and adding more information so we know what transaction is active.

diff --git a/Engine/Source/Editor/FoliageEdit/Private/FoliageEdMode.cpp b/Engine/Source/Editor/FoliageEdit/Private/FoliageEdMode.cpp
--- a/Engine/Source/Editor/FoliageEdit/Private/FoliageEdMode.cpp
+++ b/Engine/Source/Editor/FoliageEdit/Private/FoliageEdMode.cpp
@@ -513,7 +513,8 @@ void FEdModeFoliage::Exit()
 	{
 		EndTracking();
 	}
-	check(!GEditor->IsTransactionActive());
+
+	ensureMsgf(!GEditor->IsTransactionActive(), TEXT("Transaction Active when exiting foliage edit mode: '%s'"), *GEditor->GetTransactionName().ToString());
 
 	FToolkitManager::Get().CloseToolkit(Toolkit.ToSharedRef());
 	Toolkit.Reset();
