Foliage Mode: Filter out Cooked Foliage types until we have a better way of identifying public/private assets

#rnx

diff --git a/Engine/Source/Editor/FoliageEdit/Private/FoliageEdMode.cpp b/Engine/Source/Editor/FoliageEdit/Private/FoliageEdMode.cpp
--- a/Engine/Source/Editor/FoliageEdit/Private/FoliageEdMode.cpp
+++ b/Engine/Source/Editor/FoliageEdit/Private/FoliageEdMode.cpp
@@ -3180,6 +3180,12 @@ void FEdModeFoliage::PopulateFoliageMeshList()
 				continue;
 			}
 
+			//@todo_ow: have a better filter for cooked assets. For now: filter all of them out.
+			if (MeshPair.Key && MeshPair.Key->GetPackage()->HasAnyPackageFlags(PKG_Cooked))
+			{
+				continue;
+			}
+
 			int32 ElementIdx = FoliageMeshList.IndexOfByPredicate([&](const FFoliageMeshUIInfoPtr& Item)
 			{
 				return Item->Settings == MeshPair.Key;
