Foliage Mode: SupportF ExternalContentDragDropOp
- Also use ContentBrowserModule SaveAssetDialog instead of the limited SDlgPickAssetPath dialog

#rnx

diff --git a/Engine/Source/Editor/FoliageEdit/Private/FoliageEdMode.cpp b/Engine/Source/Editor/FoliageEdit/Private/FoliageEdMode.cpp
--- a/Engine/Source/Editor/FoliageEdit/Private/FoliageEdMode.cpp
+++ b/Engine/Source/Editor/FoliageEdit/Private/FoliageEdMode.cpp
@@ -3428,7 +3428,7 @@ void FEdModeFoliage::MoveSelectedFoliageToLevel(ULevel* InTargetLevel)
 	}
 }
 
-UFoliageType* FEdModeFoliage::AddFoliageAsset(UObject* InAsset)
+UFoliageType* FEdModeFoliage::AddFoliageAsset(UObject* InAsset, bool bInPlaceholderAsset)
 {
 	UFoliageType* FoliageType = nullptr;
 
@@ -3452,7 +3452,7 @@ UFoliageType* FEdModeFoliage::AddFoliageAsset(UObject* InAsset)
 				// If the world is partitioned or the world has sublevels,
 				// require the user to save the new foliage as an asset.
 				// The user will then be able to paint over all cells/sub-level.
-				FoliageType = SaveFoliageTypeObject(FoliageType);
+				FoliageType = SaveFoliageTypeObject(FoliageType, bInPlaceholderAsset);
 			}
 
 			if (FoliageType != nullptr)
@@ -3586,9 +3586,9 @@ void FEdModeFoliage::ReplaceSettingsObject(UFoliageType* OldSettings, UFoliageTy
 	PopulateFoliageMeshList();
 }
 
-UFoliageType* FEdModeFoliage::SaveFoliageTypeObject(UFoliageType* InFoliageType)
+UFoliageType* FEdModeFoliage::SaveFoliageTypeObject(UFoliageType* InFoliageType, bool bInPlaceholderAsset)
 {
-	UFoliageType* TypeToSave = FFoliageEditUtility::SaveFoliageTypeObject(InFoliageType);
+	UFoliageType* TypeToSave = FFoliageEditUtility::SaveFoliageTypeObject(InFoliageType, bInPlaceholderAsset);
 
 	if (TypeToSave != nullptr && TypeToSave != InFoliageType)
 	{
diff --git a/Engine/Source/Editor/FoliageEdit/Private/FoliageEdMode.h b/Engine/Source/Editor/FoliageEdit/Private/FoliageEdMode.h
--- a/Engine/Source/Editor/FoliageEdit/Private/FoliageEdMode.h
+++ b/Engine/Source/Editor/FoliageEdit/Private/FoliageEdMode.h
@@ -485,7 +485,7 @@ public:
 	bool IsModifierButtonPressed(const FEditorViewportClient* ViewportClient) const;
 
 	/** Add a new asset (FoliageType or StaticMesh) */
-	UFoliageType* AddFoliageAsset(UObject* InAsset);
+	UFoliageType* AddFoliageAsset(UObject* InAsset, bool bInPlaceholderAsset = false);
 
 	/** Remove a list of Foliage types */
 	bool RemoveFoliageType(UFoliageType** FoliageTypes, int32 Num);
@@ -509,7 +509,7 @@ public:
 	void ReplaceSettingsObject(UFoliageType* OldSettings, UFoliageType* NewSettings);
 
 	/** Save the foliage type object. If it isn't an asset, will prompt the user for a location to save the new asset. */
-	UFoliageType* SaveFoliageTypeObject(UFoliageType* Settings);
+	UFoliageType* SaveFoliageTypeObject(UFoliageType* Settings, bool bPlaceholderAsset = false);
 
 	void IncludeNonFoliageActors(const TArray<const UFoliageType*>& FoliageTypes, bool bOnlyCurrentLevel);
 
diff --git a/Engine/Source/Editor/FoliageEdit/Private/FoliageEditUtility.cpp b/Engine/Source/Editor/FoliageEdit/Private/FoliageEditUtility.cpp
--- a/Engine/Source/Editor/FoliageEdit/Private/FoliageEditUtility.cpp
+++ b/Engine/Source/Editor/FoliageEdit/Private/FoliageEditUtility.cpp
@@ -39,33 +39,44 @@
 #include "UObject/UObjectGlobals.h"
 #include "Widgets/DeclarativeSyntaxSupport.h"
 #include "Widgets/Notifications/SNotificationList.h"
+#include "IContentBrowserSingleton.h"
+#include "ContentBrowserModule.h"
 
 #define LOCTEXT_NAMESPACE "FoliageEdMode"
 
-UFoliageType* FFoliageEditUtility::SaveFoliageTypeObject(UFoliageType* InFoliageType)
+UFoliageType* FFoliageEditUtility::SaveFoliageTypeObject(UFoliageType* InFoliageType, bool bInPlaceholderAsset)
 {
 	UFoliageType* TypeToSave = nullptr;
 
 	if (!InFoliageType->IsAsset())
 	{
 		FString PackageName;
+		FString AssetName = InFoliageType->GetDefaultNewAssetName();
 		UObject* FoliageSource = InFoliageType->GetSource();
 		if (FoliageSource)
 		{
+			// Avoid using source name if this is a placeholder asset which is going to be replaced 
+			if (!bInPlaceholderAsset)
+			{
+				AssetName = FoliageSource->GetName() + TEXT("_FoliageType");
+			}
+
 			// Build default settings asset name and path
-			PackageName = FPackageName::GetLongPackagePath(FoliageSource->GetOutermost()->GetName()) + TEXT("/") + FoliageSource->GetName() + TEXT("_FoliageType");
+			PackageName = FPackageName::GetLongPackagePath(FoliageSource->GetOutermost()->GetName()) + TEXT("/") + AssetName;
 		}
 
-		TSharedRef<SDlgPickAssetPath> SaveFoliageTypeDialog =
-			SNew(SDlgPickAssetPath)
-			.Title(LOCTEXT("SaveFoliageTypeDialogTitle", "Choose Location for Foliage Type Asset"))
-			.DefaultAssetPath(FText::FromString(PackageName))
-			.AllowReadOnlyFolders(false);
+		FSaveAssetDialogConfig SaveAssetDialogConfig;
+		SaveAssetDialogConfig.DialogTitleOverride = LOCTEXT("SaveAssetDialogTitle", "Save Asset As");
+		SaveAssetDialogConfig.DefaultPath = FPaths::GetPath(PackageName);
+		SaveAssetDialogConfig.DefaultAssetName = AssetName;
+		SaveAssetDialogConfig.ExistingAssetPolicy = ESaveAssetDialogExistingAssetPolicy::Disallow;
 
-		if (SaveFoliageTypeDialog->ShowModal() != EAppReturnType::Cancel)
+		FContentBrowserModule& ContentBrowserModule = FModuleManager::LoadModuleChecked<FContentBrowserModule>("ContentBrowser");
+		FString SaveObjectPath = ContentBrowserModule.Get().CreateModalSaveAssetDialog(SaveAssetDialogConfig);
+		if(!SaveObjectPath.IsEmpty())
 		{
-			PackageName = SaveFoliageTypeDialog->GetFullAssetPath().ToString();
-			TypeToSave = DuplicateFoliageTypeToNewPackage(PackageName, InFoliageType);
+			FSoftObjectPath SoftObjectPath(SaveObjectPath);
+			TypeToSave = DuplicateFoliageTypeToNewPackage(SoftObjectPath.GetLongPackageName(), InFoliageType);
 		}
 	}
 	else
diff --git a/Engine/Source/Editor/FoliageEdit/Private/FoliageTypeFactory.cpp b/Engine/Source/Editor/FoliageEdit/Private/FoliageTypeFactory.cpp
--- a/Engine/Source/Editor/FoliageEdit/Private/FoliageTypeFactory.cpp
+++ b/Engine/Source/Editor/FoliageEdit/Private/FoliageTypeFactory.cpp
@@ -26,6 +26,11 @@ UObject* UFoliageType_InstancedStaticMeshFactory::FactoryCreateNew(UClass* Class
 	return NewObject<UFoliageType_InstancedStaticMesh>(InParent, Class, Name, Flags | RF_Transactional);
 }
 
+FString UFoliageType_InstancedStaticMeshFactory::GetDefaultNewAssetName() const
+{
+	return UFoliageType_InstancedStaticMesh::StaticClass()->GetDefaultObject<UFoliageType_InstancedStaticMesh>()->GetDefaultNewAssetName();
+}
+
 FText UFoliageType_InstancedStaticMeshFactory::GetToolTip() const
 {
 	return LOCTEXT("FoliageTypeStaticMeshToolTip", "Static Mesh Foliage is a foliage type that will use mesh instancing and is optimal for non-interactive foliage.");
@@ -44,6 +49,11 @@ UObject* UFoliageType_ActorFactory::FactoryCreateNew(UClass* Class, UObject* InP
 	return NewObject<UFoliageType_Actor>(InParent, Class, Name, Flags | RF_Transactional);
 }
 
+FString UFoliageType_ActorFactory::GetDefaultNewAssetName() const
+{
+	return UFoliageType_Actor::StaticClass()->GetDefaultObject<UFoliageType_Actor>()->GetDefaultNewAssetName();
+}
+
 FText UFoliageType_ActorFactory::GetToolTip() const
 {
 	return LOCTEXT("FoliageTypeActorToolTip", "Actor Foliage is a foliage type that will place blueprint/native actor instances. The cost of painting this foliage is the same as adding actors in a scene.");
diff --git a/Engine/Source/Editor/FoliageEdit/Private/FoliageTypeFactory.h b/Engine/Source/Editor/FoliageEdit/Private/FoliageTypeFactory.h
--- a/Engine/Source/Editor/FoliageEdit/Private/FoliageTypeFactory.h
+++ b/Engine/Source/Editor/FoliageEdit/Private/FoliageTypeFactory.h
@@ -27,10 +27,7 @@ class UFoliageType_InstancedStaticMeshFactory : public UFactory
 
 	// UFactory interface
 	virtual UObject* FactoryCreateNew(UClass* Class, UObject* InParent, FName Name, EObjectFlags Flags, UObject* Context, FFeedbackContext* Warn) override;
-	virtual FString GetDefaultNewAssetName() const override
-	{
-		return TEXT("NewInstancedStaticMeshFoliage");
-	}
+	virtual FString GetDefaultNewAssetName() const override;
 	virtual FText GetToolTip() const override;
 	// End of UFactory interface
 };
@@ -42,10 +39,7 @@ class UFoliageType_ActorFactory : public UFactory
 
 	// UFactory interface
 	virtual UObject* FactoryCreateNew(UClass* Class, UObject* InParent, FName Name, EObjectFlags Flags, UObject* Context, FFeedbackContext* Warn) override;
-	virtual FString GetDefaultNewAssetName() const override
-	{
-		return TEXT("NewActorFoliage");
-	}
+	virtual FString GetDefaultNewAssetName() const override;
 	virtual FText GetToolTip() const override;
 	// End of UFactory interface
 };
\ No newline at end of file
diff --git a/Engine/Source/Editor/FoliageEdit/Private/SFoliagePalette.cpp b/Engine/Source/Editor/FoliageEdit/Private/SFoliagePalette.cpp
--- a/Engine/Source/Editor/FoliageEdit/Private/SFoliagePalette.cpp
+++ b/Engine/Source/Editor/FoliageEdit/Private/SFoliagePalette.cpp
@@ -16,6 +16,7 @@
 #include "Editor/EditorEngine.h"
 #include "Engine/Blueprint.h"
 #include "Engine/StaticMesh.h"
+#include "InstancedFoliageActor.h"
 #include "FoliageEdMode.h"
 #include "FoliagePaletteCommands.h"
 #include "FoliagePaletteItem.h"
@@ -62,6 +63,7 @@
 #include "UObject/Package.h"
 #include "UObject/TopLevelAssetPath.h"
 #include "UObject/UnrealNames.h"
+#include "UObject/UObjectIterator.h"
 #include "Widgets/Images/SImage.h"
 #include "Widgets/Input/SCheckBox.h"
 #include "Widgets/Input/SComboButton.h"
@@ -80,8 +82,9 @@
 #include "Widgets/Text/STextBlock.h"
 #include "Widgets/Views/ITableRow.h"
 #include "Widgets/Views/SListView.h"
+#include "Input/DragAndDrop.h"
+#include "DragAndDrop/ExternalContentDragDropOp.h"
 
-class FDragDropEvent;
 class UObject;
 struct FGeometry;
 struct FKeyEvent;
@@ -219,6 +222,8 @@ void SFoliagePalette::Construct(const FArguments& InArgs)
 
 	FoliageEditMode->OnToolChanged.AddSP(this, &SFoliagePalette::HandleOnToolChanged);
 
+	FEditorDelegates::OnExternalContentResolved.AddSP(this, &SFoliagePalette::OnExternalContentResolved);
+
 	FFoliagePaletteCommands::Register();
 	UICommandList = MakeShareable(new FUICommandList);
 	BindCommands();
@@ -560,7 +565,7 @@ void SFoliagePalette::RefreshActivePaletteViewWidget()
 	}
 }
 
-void SFoliagePalette::AddFoliageType(const FAssetData& AssetData)
+UFoliageType* SFoliagePalette::AddFoliageType(const FAssetData& AssetData, bool bPlaceholderAsset)
 {
 	if (AddFoliageTypeCombo.IsValid())
 	{
@@ -571,7 +576,7 @@ void SFoliagePalette::AddFoliageType(const FAssetData& AssetData)
 	UObject* Asset = AssetData.GetAsset();
 	GWarn->EndSlowTask();
 
-	FoliageEditMode->AddFoliageAsset(Asset);
+	return FoliageEditMode->AddFoliageAsset(Asset, bPlaceholderAsset);
 }
 
 TSharedRef<SWidgetSwitcher> SFoliagePalette::CreatePaletteViews()
@@ -678,6 +683,11 @@ void SFoliagePalette::OnSearchTextChanged(const FText& InFilterText)
 	UpdatePalette();
 }
 
+void SFoliagePalette::AddFoliageTypePicker(const FAssetData& AssetData)
+{
+	AddFoliageType(AssetData);
+}
+
 TSharedRef<SWidget> SFoliagePalette::GetAddFoliageTypePicker()
 {
 	TArray<const UClass*> ClassFilters;
@@ -689,7 +699,7 @@ TSharedRef<SWidget> SFoliagePalette::GetAddFoliageTypePicker()
 		ClassFilters,
 		PropertyCustomizationHelpers::GetNewAssetFactoriesForClasses(ClassFilters),
 		FOnShouldFilterAsset(),
-		FOnAssetSelected::CreateSP(this, &SFoliagePalette::AddFoliageType),
+		FOnAssetSelected::CreateSP(this, &SFoliagePalette::AddFoliageTypePicker),
 		FSimpleDelegate());
 }
 
@@ -895,18 +905,71 @@ FReply SFoliagePalette::HandleFoliageDropped(const FGeometry& DropZoneGeometry,
 	TArray<FAssetData> DroppedAssetData = AssetUtil::ExtractAssetDataFromDrag(DragDropEvent);
 	if (DroppedAssetData.Num() > 0)
 	{
+		const bool bIsExternalContent = DragDropEvent.GetOperation().IsValid() && DragDropEvent.GetOperation()->IsOfType<FExternalContentDragDropOp>();
+
 		// Treat the entire drop as a transaction (in case multiples types are being added)
 		const FScopedTransaction Transaction(NSLOCTEXT("UnrealEd", "FoliageMode_DragDropTypesTransaction", "Drag-drop Foliage"));
 
+		TArray<TWeakObjectPtr<UFoliageType>> AddedFoliageTypes;
+
 		for (auto& AssetData : DroppedAssetData)
 		{
-			AddFoliageType(AssetData);
+			if (UFoliageType* FoliageType = AddFoliageType(AssetData, bIsExternalContent))
+			{
+				AddedFoliageTypes.Add(FoliageType);
+			}
+		}
+
+		if (AddedFoliageTypes.Num() > 0 && bIsExternalContent)
+		{
+			TSharedPtr<const FExternalContentDragDropOp> DragDropOp = StaticCastSharedPtr<const FExternalContentDragDropOp>(DragDropEvent.GetOperation());
+			ExternalContentFoliageTypes.Add(DragDropOp->GetGuid(), AddedFoliageTypes);
 		}
 	}
 
 	return FReply::Handled();
 }
 
+void SFoliagePalette::OnExternalContentResolved(const FGuid& Identifier, const FAssetData& PlaceHolderAsset, const FAssetData& ResolvedAsset)
+{
+	if (TArray<TWeakObjectPtr<UFoliageType>>* FoliageTypes = ExternalContentFoliageTypes.Find(Identifier))
+	{
+		TArray<UFoliageType*> ModifiedFoliageTypes;
+
+		for (TWeakObjectPtr<UFoliageType> FoliageTypePtr : *FoliageTypes)
+		{
+			if (UFoliageType_InstancedStaticMesh* FoliageType = Cast<UFoliageType_InstancedStaticMesh>(FoliageTypePtr.Get()))
+			{
+				FAssetData AssetData(FoliageType->GetStaticMesh());
+				if (AssetData == PlaceHolderAsset)
+				{
+					if (UStaticMesh* ResolvedMesh = Cast<UStaticMesh>(ResolvedAsset.GetAsset()))
+					{
+						FoliageType->Modify();
+						FoliageType->SetStaticMesh(ResolvedMesh);
+						ModifiedFoliageTypes.Add(FoliageType);
+					}
+				}
+			}
+		}
+		ExternalContentFoliageTypes.Remove(Identifier);
+
+		if (ModifiedFoliageTypes.Num() > 0)
+		{
+			for (TObjectIterator<AInstancedFoliageActor> It(RF_ClassDefaultObject, /** bIncludeDerivedClasses */ true, /** InternalExcludeFlags */ EInternalObjectFlags::Garbage); It; ++It)
+			{
+				if (It->GetWorld() != nullptr)
+				{
+					for (UFoliageType* ModifiedFoliageType : ModifiedFoliageTypes)
+					{
+						It->NotifyFoliageTypeChanged(ModifiedFoliageType, true);
+					}
+				}
+			}
+		}
+	}
+}
+
 //	CONTEXT MENU
 
 TSharedPtr<SWidget> SFoliagePalette::ConstructFoliageTypeContextMenu()
diff --git a/Engine/Source/Editor/FoliageEdit/Private/SFoliagePalette.h b/Engine/Source/Editor/FoliageEdit/Private/SFoliagePalette.h
--- a/Engine/Source/Editor/FoliageEdit/Private/SFoliagePalette.h
+++ b/Engine/Source/Editor/FoliageEdit/Private/SFoliagePalette.h
@@ -90,9 +90,10 @@ public:
 	FText GetSearchText() const;
 
 	/** Adds the foliage type asset to the instanced foliage actor's list of types. */
-	void AddFoliageType(const FAssetData& AssetData);
-
+	UFoliageType* AddFoliageType(const FAssetData& AssetData, bool bPlaceholderAsset = false);
+	
 private:	// GENERAL
+	void AddFoliageTypePicker(const FAssetData& AssetData);
 
 	/** Binds commands used by the palette */
 	void BindCommands();
@@ -122,6 +123,8 @@ private:	// GENERAL
 	/** Handler to trigger a refresh of the details view when the active tool changes */
 	void HandleOnToolChanged();
 
+	void OnExternalContentResolved(const FGuid& Identifier, const FAssetData& PlaceholderAsset, const FAssetData& ResolvedAsset);
+
 	/** Sets the view mode of the palette */
 	void SetViewMode(EFoliagePaletteViewMode::Type NewViewMode);
 
@@ -322,6 +325,8 @@ private:
 
 	FEdModeFoliage* FoliageEditMode;
 
+	TMap<FGuid, TArray<TWeakObjectPtr<UFoliageType>>> ExternalContentFoliageTypes;
+
 	bool bItemsNeedRebuild : 1;
 	bool bIsUneditableFoliageTypeSelected : 1;
 	
diff --git a/Engine/Source/Editor/FoliageEdit/Public/FoliageEditUtility.h b/Engine/Source/Editor/FoliageEdit/Public/FoliageEditUtility.h
--- a/Engine/Source/Editor/FoliageEdit/Public/FoliageEditUtility.h
+++ b/Engine/Source/Editor/FoliageEdit/Public/FoliageEditUtility.h
@@ -20,7 +20,7 @@ class FFoliageEditUtility
 {
 public:
 	/** Save the foliage type object. If it isn't an asset, will prompt the user for a location to save the new asset. */
-	static FOLIAGEEDIT_API UFoliageType* SaveFoliageTypeObject(UFoliageType* Settings); 
+	static FOLIAGEEDIT_API UFoliageType* SaveFoliageTypeObject(UFoliageType* Settings, bool bPlaceholderAsset = false); 
 
 	/**
 	 * Copies the provided FoliageType to a new asset
diff --git a/Engine/Source/Runtime/Foliage/Public/FoliageType.h b/Engine/Source/Runtime/Foliage/Public/FoliageType.h
--- a/Engine/Source/Runtime/Foliage/Public/FoliageType.h
+++ b/Engine/Source/Runtime/Foliage/Public/FoliageType.h
@@ -116,6 +116,10 @@ class UFoliageType : public UObject
 	FOLIAGE_API FVector3f GetRandomScale() const;
 	
 #if WITH_EDITOR
+	virtual FString GetDefaultNewAssetName() const
+	{
+		return TEXT("NewFoliage");
+	}
 
 	virtual void SetSource(UObject* InSource) PURE_VIRTUAL(UFoliageType::SetSource, );
 	virtual void UpdateBounds() {}
diff --git a/Engine/Source/Runtime/Foliage/Public/FoliageType_Actor.h b/Engine/Source/Runtime/Foliage/Public/FoliageType_Actor.h
--- a/Engine/Source/Runtime/Foliage/Public/FoliageType_Actor.h
+++ b/Engine/Source/Runtime/Foliage/Public/FoliageType_Actor.h
@@ -27,6 +27,11 @@ class UFoliageType_Actor : public UFoliageType
 	virtual UObject* GetSource() const override { return ActorClass; }
 
 #if WITH_EDITOR
+	virtual FString GetDefaultNewAssetName() const override
+	{
+		return TEXT("NewActorFoliage");
+	}
+
 	virtual void UpdateBounds();
 	virtual bool IsSourcePropertyChange(const FProperty* Property) const override
 	{
diff --git a/Engine/Source/Runtime/Foliage/Public/FoliageType_InstancedStaticMesh.h b/Engine/Source/Runtime/Foliage/Public/FoliageType_InstancedStaticMesh.h
--- a/Engine/Source/Runtime/Foliage/Public/FoliageType_InstancedStaticMesh.h
+++ b/Engine/Source/Runtime/Foliage/Public/FoliageType_InstancedStaticMesh.h
@@ -41,6 +41,11 @@ class UFoliageType_InstancedStaticMesh : public UFoliageType
 	virtual UObject* GetSource() const override;
 
 #if WITH_EDITOR
+	virtual FString GetDefaultNewAssetName() const override
+	{
+		return TEXT("NewInstancedStaticMeshFoliage");
+	}
+
 	virtual void UpdateBounds() override;
 	virtual bool IsSourcePropertyChange(const FProperty* Property) const override
 	{
diff --git a/Engine/Source/Runtime/Foliage/Public/InstancedFoliageActor.h b/Engine/Source/Runtime/Foliage/Public/InstancedFoliageActor.h
--- a/Engine/Source/Runtime/Foliage/Public/InstancedFoliageActor.h
+++ b/Engine/Source/Runtime/Foliage/Public/InstancedFoliageActor.h
@@ -270,7 +270,7 @@ public:
 	static FOLIAGE_API bool HasFoliageAttached(UActorComponent* InComponent);
 
 	/* Called to notify InstancedFoliageActor that a UFoliageType has been modified */
-	void NotifyFoliageTypeChanged(UFoliageType* FoliageType, bool bSourceChanged);
+	FOLIAGE_API void NotifyFoliageTypeChanged(UFoliageType* FoliageType, bool bSourceChanged);
 	void NotifyFoliageTypeWillChange(UFoliageType* FoliageType);
 
 	DECLARE_EVENT_OneParam(AInstancedFoliageActor, FOnFoliageTypeMeshChanged, UFoliageType*);
