MoveSelectedFoliageToActorEditorContext: Add missing transaction scope

#rnx

diff --git a/Engine/Source/Editor/FoliageEdit/Private/FoliageEdMode.cpp b/Engine/Source/Editor/FoliageEdit/Private/FoliageEdMode.cpp
--- a/Engine/Source/Editor/FoliageEdit/Private/FoliageEdMode.cpp
+++ b/Engine/Source/Editor/FoliageEdit/Private/FoliageEdMode.cpp
@@ -3327,6 +3327,7 @@ bool FEdModeFoliage::IsModifierButtonPressed(const FEditorViewportClient* Viewpo
 
 void FEdModeFoliage::MoveSelectedFoliageToActorEditorContext()
 {
+	const FScopedTransaction Transaction(NSLOCTEXT("UnrealEd", "MoveSelectedFoliageToActorEditorContext", "Move Selected Foliage to Actor Editor Context"));
 	AInstancedFoliageActor::MoveSelectedInstancesToActorEditorContext(GetWorld());
 }
 
