Foliage Erase: Fix erasing when using grid foliage (World Partition)

- DesiredInstanceCount is global to brush sphere so we need to calculate a Potential instances to remove global to the sphere and not local to each intersecting IFA

#rnx

diff --git a/Engine/Source/Editor/FoliageEdit/Private/FoliageEdMode.cpp b/Engine/Source/Editor/FoliageEdit/Private/FoliageEdMode.cpp
--- a/Engine/Source/Editor/FoliageEdit/Private/FoliageEdMode.cpp
+++ b/Engine/Source/Editor/FoliageEdit/Private/FoliageEdMode.cpp
@@ -1623,36 +1623,67 @@ void FEdModeFoliage::RemoveInstancesForBrush(UWorld* InWorld, const UFoliageType
 {
 	SCOPE_CYCLE_COUNTER(STAT_FoliageRemoveInstanceBrush);
 
-	auto RemovingInstances = [this, &BrushSphere, &DesiredInstanceCount, &Pressure](AInstancedFoliageActor* IFA, FFoliageInfo* FoliageInfo, const UFoliageType* FoliageType) {
-		TArray<int32> PotentialInstancesToRemove;
-		FoliageInfo->GetInstancesInsideSphere(BrushSphere, PotentialInstancesToRemove);
-		if (PotentialInstancesToRemove.Num() == 0)
-		{
-			return true;
-		}
+	struct FInstancesToRemove
+	{
+		FFoliageInfo* FoliageInfo;
+		TArray<int32> InstancesToRemove;
+	};
+	TArray<FInstancesToRemove> PotentialInstancesToRemove;
+	int32 PotentialInstancesToRemoveCount = 0;
 
-		int32 InstancesToRemove = FMath::RoundToInt((float)(PotentialInstancesToRemove.Num() - DesiredInstanceCount) * Pressure);
-		if (InstancesToRemove <= 0)
+	// Get Brush intersecting instances per FoliageInfo (per IFA)
+	ForEachFoliageInfo(InWorld, Settings, BrushSphere, [&BrushSphere, &PotentialInstancesToRemove, &PotentialInstancesToRemoveCount](AInstancedFoliageActor* IFA, FFoliageInfo* FoliageInfo, const UFoliageType* FoliageType)
+	{
+		TArray<int32> InstancesInsideSphere;
+		FoliageInfo->GetInstancesInsideSphere(BrushSphere, InstancesInsideSphere);
+		if (InstancesInsideSphere.Num() > 0)
 		{
-			return true;
+			PotentialInstancesToRemoveCount += InstancesInsideSphere.Num();
+			PotentialInstancesToRemove.Add({ FoliageInfo, MoveTemp(InstancesInsideSphere) });
 		}
+		return true;
+	});
+
+	// Calculate number of Instances to remove based on desired instance count
+	const int32 InstancesToRemoveCount = FMath::RoundToInt((float)(PotentialInstancesToRemoveCount - DesiredInstanceCount) * Pressure);
+	if (InstancesToRemoveCount <= 0)
+	{
+		return;
+	}
 
-		int32 InstancesToKeep = PotentialInstancesToRemove.Num() - InstancesToRemove;
-		if (InstancesToKeep > 0)
+	// Calculate InstancesToKeep
+	const int32 InstancesToKeepCount = PotentialInstancesToRemoveCount - InstancesToRemoveCount;
+	
+	// Remove InstancesToKeep from the PotentialInstancesToRemove randomly so that they don't get removed
+	for (int32 i = 0; i < InstancesToKeepCount; i++)
+	{
+		const int32 RemoveIndex = FMath::Rand() % PotentialInstancesToRemoveCount;
+		int32 StartIndex = 0;
+		
+		// Iterate through IFAs to find the RemoveIndex
+		for (auto& [FoliageInfo, InstancesToRemove] : PotentialInstancesToRemove)
 		{
-			// Remove InstancesToKeep random PotentialInstancesToRemove from the array to leave those PotentialInstancesToRemove behind, and delete all the rest
-			for (int32 i = 0; i < InstancesToKeep; i++)
+			const int32 LocalRemoveIndex = RemoveIndex - StartIndex;
+			if (InstancesToRemove.IsValidIndex(LocalRemoveIndex))
 			{
-				PotentialInstancesToRemove.RemoveAtSwap(FMath::Rand() % PotentialInstancesToRemove.Num(), 1, false);
+				InstancesToRemove.RemoveAtSwap(LocalRemoveIndex, 1, false);
+				break;
 			}
+			StartIndex += InstancesToRemove.Num();
 		}
+		PotentialInstancesToRemoveCount--;
+	}
+	
+	FFoliagePaintingGeometryFilter GeometryFilterFunc(UISettings);
 
-		FFoliagePaintingGeometryFilter GeometryFilterFunc(UISettings);
+	// Filter PotentialInstancesToRemove
+	for (auto& [FoliageInfo, InstancesToRemove] : PotentialInstancesToRemove)
+	{
+		AInstancedFoliageActor* IFA = FoliageInfo->IFA;
 
-		// Filter PotentialInstancesToRemove
-		for (int32 Idx = 0; Idx < PotentialInstancesToRemove.Num(); Idx++)
+		for (int32 Idx = 0; Idx < InstancesToRemove.Num(); Idx++)
 		{
-			auto BaseId = FoliageInfo->Instances[PotentialInstancesToRemove[Idx]].BaseId;
+			auto BaseId = FoliageInfo->Instances[InstancesToRemove[Idx]].BaseId;
 			auto BasePtr = IFA->InstanceBaseCache.GetInstanceBasePtr(BaseId);
 			UPrimitiveComponent* Base = Cast<UPrimitiveComponent>(BasePtr.Get());
 
@@ -1660,24 +1691,21 @@ void FEdModeFoliage::RemoveInstancesForBrush(UWorld* InWorld, const UFoliageType
 			if (Base && !GeometryFilterFunc(Base))
 			{
 				// Instance should not be removed, so remove it from the removal list.
-				PotentialInstancesToRemove.RemoveAtSwap(Idx, 1, false);
+				InstancesToRemove.RemoveAtSwap(Idx, 1, false);
 				Idx--;
 			}
 		}
 
-		// Remove PotentialInstancesToRemove to reduce it to desired count
-		if (PotentialInstancesToRemove.Num() > 0)
+		// Remove InstancesToRemove to reduce it to desired count
+		if (InstancesToRemove.Num() > 0)
 		{
-			CurrentFoliageTraceBrushAffectedIFAs.Add(IFA);
+			CurrentFoliageTraceBrushAffectedIFAs.Add(FoliageInfo->IFA);
 
-			FoliageInfo->RemoveInstances(PotentialInstancesToRemove, false);
+			FoliageInfo->RemoveInstances(InstancesToRemove, false);
 		}
-		return true;
-	};
-	ForEachFoliageInfo(InWorld, Settings, BrushSphere, RemovingInstances);
+	}
 }
 
-
 void FEdModeFoliage::SelectInstanceAtLocation(UWorld* InWorld, const UFoliageType* Settings, const FVector& Location, bool bSelect)
 {
 	FEdModeFoliageSelectionUpdate Scope(this);
