Foliage Mode to use new IsAssetVisible check instead of skipping cooked packages

diff --git a/Engine/Source/Editor/FoliageEdit/Private/FoliageEdMode.cpp b/Engine/Source/Editor/FoliageEdit/Private/FoliageEdMode.cpp
--- a/Engine/Source/Editor/FoliageEdit/Private/FoliageEdMode.cpp
+++ b/Engine/Source/Editor/FoliageEdit/Private/FoliageEdMode.cpp
@@ -7,6 +7,7 @@
 #include "AssetRegistry/AssetData.h"
 #include "AssetRegistry/AssetRegistryModule.h"
 #include "AssetRegistry/IAssetRegistry.h"
+#include "IAssetTools.h"
 #include "CollisionQueryParams.h"
 #include "CollisionShape.h"
 #include "Components/ActorComponent.h"
@@ -3208,8 +3209,7 @@ void FEdModeFoliage::PopulateFoliageMeshList()
 				continue;
 			}
 
-			//@todo_ow: have a better filter for cooked assets. For now: filter all of them out.
-			if (MeshPair.Key && MeshPair.Key->GetPackage()->HasAnyPackageFlags(PKG_Cooked))
+			if (!IAssetTools::Get().IsAssetVisible(FAssetData(MeshPair.Key)))
 			{
 				continue;
 			}
