Unhide "painting" category when viewing the asset details of a foliage type.

diff --git a/Engine/Source/Editor/FoliageEdit/Private/FoliageTypeDetails.cpp b/Engine/Source/Editor/FoliageEdit/Private/FoliageTypeDetails.cpp
--- a/Engine/Source/Editor/FoliageEdit/Private/FoliageTypeDetails.cpp
+++ b/Engine/Source/Editor/FoliageEdit/Private/FoliageTypeDetails.cpp
@@ -97,8 +97,8 @@ void FFoliageTypeDetails::CustomizeDetails(IDetailLayoutBuilder& DetailBuilder)
 	const static FName InstanceSettingsName("InstanceSettings");
 
 	FFoliageTypeCustomizationHelpers::HideFoliageCategory(DetailBuilder, ReapplyName);
-	FFoliageTypeCustomizationHelpers::HideFoliageCategory(DetailBuilder, PaintingName);
 
+	AddSubcategoryProperties(DetailBuilder, PaintingName);
 	AddSubcategoryProperties(DetailBuilder, PlacementName);
 
 	// If enabled, show the properties for procedural placement
