Move OnReflectSelectionInPalette command from foliage palette widget to foliage ed mode. This enables the command to be called in the viewport, etc. instead of being limited to the focus of the foliage palette widget.

diff --git a/Engine/Source/Editor/FoliageEdit/Private/FoliageEdMode.cpp b/Engine/Source/Editor/FoliageEdit/Private/FoliageEdMode.cpp
--- a/Engine/Source/Editor/FoliageEdit/Private/FoliageEdMode.cpp
+++ b/Engine/Source/Editor/FoliageEdit/Private/FoliageEdMode.cpp
@@ -408,6 +408,10 @@ void FEdModeFoliage::BindCommands()
 	{
 		return UISettings.GetPaintBucketToolSelected();
 	}));
+
+	UICommandList->MapAction(
+		Commands.ReflectSelectionInPalette,
+		FExecuteAction::CreateSP(this, &FEdModeFoliage::OnReflectSelectionInPalette));
 }
 
 bool FEdModeFoliage::CurrentToolUsesBrush() const
@@ -760,6 +764,11 @@ void FEdModeFoliage::OnSetPlace()
 	HandleToolChanged();
 }
 
+void FEdModeFoliage::OnReflectSelectionInPalette()
+{
+	StaticCastSharedPtr<FFoliageEdModeToolkit>(Toolkit)->ReflectSelectionInPalette();
+}
+
 bool FEdModeFoliage::DisallowMouseDeltaTracking() const
 {
 	// We never want to use the mouse delta tracker while painting
diff --git a/Engine/Source/Editor/FoliageEdit/Private/FoliageEdMode.h b/Engine/Source/Editor/FoliageEdit/Private/FoliageEdMode.h
--- a/Engine/Source/Editor/FoliageEdit/Private/FoliageEdMode.h
+++ b/Engine/Source/Editor/FoliageEdit/Private/FoliageEdMode.h
@@ -585,6 +585,9 @@ public:
 	/** Sets the tool mode to Place Single Instance*/
 	void OnSetPlace();
 
+	/** Handle reflecting selected foliage types in the FoliagePalette */
+	void OnReflectSelectionInPalette();
+
 	/** Remove currently selected instances*/
 	void RemoveSelectedInstances(UWorld* InWorld);
 
diff --git a/Engine/Source/Editor/FoliageEdit/Private/FoliageEdModeToolkit.cpp b/Engine/Source/Editor/FoliageEdit/Private/FoliageEdModeToolkit.cpp
--- a/Engine/Source/Editor/FoliageEdit/Private/FoliageEdModeToolkit.cpp
+++ b/Engine/Source/Editor/FoliageEdit/Private/FoliageEdModeToolkit.cpp
@@ -57,6 +57,11 @@ void FFoliageEdModeToolkit::NotifyFoliageTypeMeshChanged(class UFoliageType* Fol
 	FoliageEdWidget->NotifyFoliageTypeMeshChanged(FoliageType);
 }
 
+void FFoliageEdModeToolkit::ReflectSelectionInPalette()
+{
+	FoliageEdWidget->ReflectSelectionInPalette();
+}
+
 void FFoliageEdModeToolkit::GetToolPaletteNames(TArray<FName>& InPaletteName) const
 {
 	InPaletteName = FoliagePaletteNames;
diff --git a/Engine/Source/Editor/FoliageEdit/Private/FoliageEdModeToolkit.h b/Engine/Source/Editor/FoliageEdit/Private/FoliageEdModeToolkit.h
--- a/Engine/Source/Editor/FoliageEdit/Private/FoliageEdModeToolkit.h
+++ b/Engine/Source/Editor/FoliageEdit/Private/FoliageEdModeToolkit.h
@@ -37,6 +37,7 @@ public:
 
 	void RefreshFullList();
 	void NotifyFoliageTypeMeshChanged(class UFoliageType* FoliageType);
+	void ReflectSelectionInPalette();
 
 private:
 	TSharedPtr< class SFoliageEdit > FoliageEdWidget;
diff --git a/Engine/Source/Editor/FoliageEdit/Private/FoliageEditActions.cpp b/Engine/Source/Editor/FoliageEdit/Private/FoliageEditActions.cpp
--- a/Engine/Source/Editor/FoliageEdit/Private/FoliageEditActions.cpp
+++ b/Engine/Source/Editor/FoliageEdit/Private/FoliageEditActions.cpp
@@ -31,6 +31,8 @@ void FFoliageEditCommands::RegisterCommands()
 	UI_COMMAND( SetNoSettings, "Hide Details", "Hide details.", EUserInterfaceActionType::ToggleButton, FInputChord() );
 	UI_COMMAND( SetPaintSettings, "Show Painting settings", "Show painting settings.", EUserInterfaceActionType::ToggleButton, FInputChord() );
 	UI_COMMAND( SetClusterSettings, "Show Instance settings", "Show settings for placed instances.", EUserInterfaceActionType::ToggleButton, FInputChord() );
+
+	UI_COMMAND(ReflectSelectionInPalette, "Reflect Selection in Palette.", "Select foliage type in palette based on selected foliage instances.", EUserInterfaceActionType::None, FInputChord());
 }
 
 #undef LOCTEXT_NAMESPACE
diff --git a/Engine/Source/Editor/FoliageEdit/Private/FoliageEditActions.h b/Engine/Source/Editor/FoliageEdit/Private/FoliageEditActions.h
--- a/Engine/Source/Editor/FoliageEdit/Private/FoliageEditActions.h
+++ b/Engine/Source/Editor/FoliageEdit/Private/FoliageEditActions.h
@@ -54,6 +54,9 @@ public:
 	TSharedPtr< FUICommandInfo > SetPaintSettings;
 	TSharedPtr< FUICommandInfo > SetClusterSettings;
 
+	/** FoliagePalette commands*/
+	TSharedPtr< FUICommandInfo > ReflectSelectionInPalette;
+
 	/** FoliageType commands */
 	TSharedPtr< FUICommandInfo > RemoveFoliageType;
 	TSharedPtr< FUICommandInfo > ShowFoliageTypeInCB;
diff --git a/Engine/Source/Editor/FoliageEdit/Private/FoliagePaletteCommands.cpp b/Engine/Source/Editor/FoliageEdit/Private/FoliagePaletteCommands.cpp
--- a/Engine/Source/Editor/FoliageEdit/Private/FoliagePaletteCommands.cpp
+++ b/Engine/Source/Editor/FoliageEdit/Private/FoliagePaletteCommands.cpp
@@ -14,7 +14,6 @@ void FFoliagePaletteCommands::RegisterCommands()
 	UI_COMMAND(DeactivateFoliageType, "Deactivate", "Sets the selected foliage types in the palette as inactive (i.e. excluded in brush actions).", EUserInterfaceActionType::Button, FInputChord());
 	UI_COMMAND(RemoveFoliageType, "Remove", "Remove this foliage type from the palette. Removes all associated instances as well.", EUserInterfaceActionType::Button, FInputChord(EKeys::Delete));
 	UI_COMMAND(ShowFoliageTypeInCB, "Show in Content Browser", "Show asset in Content Browser.", EUserInterfaceActionType::Button, FInputChord());
-	UI_COMMAND(ReflectSelectionInPalette, "Reflect Selection in Palette.", "Select foliage type in palette based on selected foliage instances.", EUserInterfaceActionType::None, FInputChord());
 	UI_COMMAND(SelectAllInstances, "Select All Instances", "Select all instances of this foliage type (must be in a selection mode).", EUserInterfaceActionType::Button, FInputChord());
 	UI_COMMAND(DeselectAllInstances, "Deselect All Instances", "Deselect all instances of this foliage type (must be in a selection mode).", EUserInterfaceActionType::Button, FInputChord());
 	UI_COMMAND(SelectInvalidInstances, "Select Invalid Instances", "Select all instances of this foliage type that are off ground (must be in a selection mode).", EUserInterfaceActionType::Button, FInputChord());
diff --git a/Engine/Source/Editor/FoliageEdit/Private/FoliagePaletteCommands.h b/Engine/Source/Editor/FoliageEdit/Private/FoliagePaletteCommands.h
--- a/Engine/Source/Editor/FoliageEdit/Private/FoliagePaletteCommands.h
+++ b/Engine/Source/Editor/FoliageEdit/Private/FoliagePaletteCommands.h
@@ -35,7 +35,6 @@ public:
 	TSharedPtr< FUICommandInfo > DeactivateFoliageType;
 	TSharedPtr< FUICommandInfo > RemoveFoliageType;
 	TSharedPtr< FUICommandInfo > ShowFoliageTypeInCB;
-	TSharedPtr< FUICommandInfo > ReflectSelectionInPalette;
 	TSharedPtr< FUICommandInfo > SelectAllInstances;
 	TSharedPtr< FUICommandInfo > DeselectAllInstances;
 	TSharedPtr< FUICommandInfo > SelectInvalidInstances;
diff --git a/Engine/Source/Editor/FoliageEdit/Private/SFoliageEdit.cpp b/Engine/Source/Editor/FoliageEdit/Private/SFoliageEdit.cpp
--- a/Engine/Source/Editor/FoliageEdit/Private/SFoliageEdit.cpp
+++ b/Engine/Source/Editor/FoliageEdit/Private/SFoliageEdit.cpp
@@ -711,6 +711,11 @@ void SFoliageEdit::NotifyFoliageTypeMeshChanged(UFoliageType* FoliageType)
 	FoliagePalette->UpdateThumbnailForType(FoliageType);
 }
 
+void SFoliageEdit::ReflectSelectionInPalette()
+{
+	FoliagePalette->ReflectSelectionInPalette();
+}
+
 bool SFoliageEdit::IsFoliageEditorEnabled() const
 {
 	ErrorText->SetError(GetFoliageEditorErrorText());
diff --git a/Engine/Source/Editor/FoliageEdit/Private/SFoliageEdit.h b/Engine/Source/Editor/FoliageEdit/Private/SFoliageEdit.h
--- a/Engine/Source/Editor/FoliageEdit/Private/SFoliageEdit.h
+++ b/Engine/Source/Editor/FoliageEdit/Private/SFoliageEdit.h
@@ -40,6 +40,9 @@ public:
 	/** Notifies the widget that the mesh assigned to a foliage type in the list has changed */
 	void NotifyFoliageTypeMeshChanged(UFoliageType* FoliageType);
 
+	/** Notifies the widget to reflect its selected foliage types based on the selected foliage instances */
+	void ReflectSelectionInPalette();
+
 	/** Gets FoliageEditMode. Used by the cluster details to notify changes */
 	class FEdModeFoliage* GetFoliageEditMode() const { return FoliageEditMode; }
 
diff --git a/Engine/Source/Editor/FoliageEdit/Private/SFoliagePalette.cpp b/Engine/Source/Editor/FoliageEdit/Private/SFoliagePalette.cpp
--- a/Engine/Source/Editor/FoliageEdit/Private/SFoliagePalette.cpp
+++ b/Engine/Source/Editor/FoliageEdit/Private/SFoliagePalette.cpp
@@ -533,10 +533,6 @@ void SFoliagePalette::BindCommands()
 		Commands.ShowFoliageTypeInCB,
 		FExecuteAction::CreateSP(this, &SFoliagePalette::OnShowFoliageTypeInCB));
 
-	UICommandList->MapAction(
-		Commands.ReflectSelectionInPalette,
-		FExecuteAction::CreateSP(this, &SFoliagePalette::OnReflectSelectionInPalette));
-
 	UICommandList->MapAction(
 		Commands.SelectAllInstances,
 		FExecuteAction::CreateSP(this, &SFoliagePalette::OnSelectAllInstances),
@@ -1263,7 +1259,7 @@ void SFoliagePalette::OnShowFoliageTypeInCB()
 	}
 }
 
-void SFoliagePalette::OnReflectSelectionInPalette()
+void SFoliagePalette::ReflectSelectionInPalette()
 {
 	TArray<const UFoliageType*> SelectedFoliageTypes;
 	FoliageEditMode->GetSelectedInstanceFoliageTypes(SelectedFoliageTypes);
diff --git a/Engine/Source/Editor/FoliageEdit/Private/SFoliagePalette.h b/Engine/Source/Editor/FoliageEdit/Private/SFoliagePalette.h
--- a/Engine/Source/Editor/FoliageEdit/Private/SFoliagePalette.h
+++ b/Engine/Source/Editor/FoliageEdit/Private/SFoliagePalette.h
@@ -91,6 +91,9 @@ public:
 
 	/** Adds the foliage type asset to the instanced foliage actor's list of types. */
 	UFoliageType* AddFoliageType(const FAssetData& AssetData, bool bPlaceholderAsset = false);
+
+	/** Updates the selection of foliage types in the palette based on the selected foliage instances. */
+	void ReflectSelectionInPalette();
 	
 private:	// GENERAL
 	void AddFoliageTypePicker(const FAssetData& AssetData);
@@ -212,9 +215,6 @@ private:	// CONTEXT MENU
 	/** @return Whether selecting instances is currently possible */
 	bool CanSelectInstances() const;
 
-	/** Handler for 'Reflect Selection in Palette ' command */
-	void OnReflectSelectionInPalette();
-
 	/** Selects Foliage Type in palette */
 	void SelectFoliageTypesInPalette(const TArray<const UFoliageType*>& FoliageTypes);
 
