Fix crashes when snapping foliage to ground or deleting them while they're being moved.

diff --git a/Engine/Source/Editor/FoliageEdit/Private/FoliageEdMode.cpp b/Engine/Source/Editor/FoliageEdit/Private/FoliageEdMode.cpp
--- a/Engine/Source/Editor/FoliageEdit/Private/FoliageEdMode.cpp
+++ b/Engine/Source/Editor/FoliageEdit/Private/FoliageEdMode.cpp
@@ -2091,6 +2091,18 @@ void FEdModeFoliage::RemoveSelectedInstances(UWorld* InWorld)
 				if (FoliageInfo.SelectedIndices.Num() > 0)
 				{
 					TArray<int32> InstancesToDelete = FoliageInfo.SelectedIndices.Array();
+
+					// Make sure that any moving instances to be removed are added back to the hash.
+					if (bMoving)
+					{
+						TArray<int32> MovingInstancesToDelete = FoliageInfo.SelectedIndices.Intersect(FoliageInfo.MovingInstances).Array();
+
+						if (MovingInstancesToDelete.Num())
+						{
+							FoliageInfo.PostMoveInstances(MovingInstancesToDelete, /*bFinished*/true);
+						}
+					}
+
 					FoliageInfo.RemoveInstances(InstancesToDelete, true);
 
 					OnInstanceCountUpdated(FoliageType);
@@ -3093,6 +3105,17 @@ void FEdModeFoliage::SnapSelectedInstancesToGround(UWorld* InWorld)
 						bFoundSelection = true;
 					}
 
+					TArray<int32> MovingInstances = FoliageInfo.SelectedIndices.Intersect(FoliageInfo.MovingInstances).Array();
+					TArray<int32> NonMovingInstances = FoliageInfo.SelectedIndices.Difference(FoliageInfo.MovingInstances).Array();
+
+					// Call PostMove on already moving instances to add them back to the hash.
+					if (MovingInstances.Num())
+					{
+						check(bMoving);
+						FoliageInfo.PostMoveInstances(MovingInstances, /*bFinished=*/false);
+					}
+
+					// Call PreMove on all snapping instances.
 					FoliageInfo.PreMoveInstances(SelectedIndices);
 
 					for (int32 InstanceIndex : SelectedIndices)
@@ -3100,7 +3123,16 @@ void FEdModeFoliage::SnapSelectedInstancesToGround(UWorld* InWorld)
 						bMovedInstance |= SnapInstanceToGround(IFA, FoliageType, FoliageInfo, InstanceIndex);
 					}
 
-					FoliageInfo.PostMoveInstances(SelectedIndices);
+					if (MovingInstances.Num())
+					{
+						FoliageInfo.PostMoveInstances(MovingInstances, /*bFinished=*/false);
+						FoliageInfo.PreMoveInstances(MovingInstances);
+					}
+
+					if (NonMovingInstances.Num())
+					{
+						FoliageInfo.PostMoveInstances(NonMovingInstances, /*bFinished=*/true);
+					}
 				}
 				return true; // continue iteration
 			});
