Clear CurrentFoliageTraceBrushAffectedIFAs after adding instances with the foliage fill tool.

diff --git a/Engine/Source/Editor/FoliageEdit/Private/FoliageEdMode.cpp b/Engine/Source/Editor/FoliageEdit/Private/FoliageEdMode.cpp
--- a/Engine/Source/Editor/FoliageEdit/Private/FoliageEdMode.cpp
+++ b/Engine/Source/Editor/FoliageEdit/Private/FoliageEdMode.cpp
@@ -2993,6 +2993,8 @@ void FEdModeFoliage::ApplyPaintBucket_Add(AActor* Actor)
 		//
 		OnInstanceCountUpdated(Settings);
 	}
+
+	CurrentFoliageTraceBrushAffectedIFAs.Empty();
 }
 
 bool FEdModeFoliage::GetStaticMeshVertexColorForHit(const UStaticMeshComponent* InStaticMeshComponent, int32 InTriangleIndex, const FVector& InHitLocation, FColor& OutVertexColor)
@@ -4126,4 +4128,4 @@ void FFoliageUISettings::Save()
 	GConfig->SetInt(TEXT("FoliageEdit"), TEXT("SingleInstantiationPlacementMode"), (int32)SingleInstantiationPlacementMode, GEditorPerProjectIni);
 }
 
-#undef LOCTEXT_NAMESPACE
+#undef LOCTEXT_NAMESPACE
\ No newline at end of file
