Optimize FoliageEditorSubsystem to accumulate actor moved events and perform updates once per tick

diff --git a/Engine/Source/Editor/FoliageEdit/Private/FoliageEditorSubsystem.cpp b/Engine/Source/Editor/FoliageEdit/Private/FoliageEditorSubsystem.cpp
--- a/Engine/Source/Editor/FoliageEdit/Private/FoliageEditorSubsystem.cpp
+++ b/Engine/Source/Editor/FoliageEdit/Private/FoliageEditorSubsystem.cpp
@@ -2,6 +2,7 @@
 
 #include "FoliageEditorSubsystem.h"
 #include "CoreGlobals.h"
+#include "Editor.h"
 #include "Engine/Engine.h"
 #include "Engine/World.h"
 #include "EngineUtils.h"
@@ -43,6 +44,84 @@ void UFoliageEditorSubsystem::Deinitialize()
 	Super::Deinitialize();
 }
 
+void UFoliageEditorSubsystem::Tick(float DeltaTime)
+{
+	for (const TPair<const UWorld*, TSet<AActor*>>& WorldActorsPair : ActorsPendingMovementUpdatePerWorld)
+	{
+		const UWorld* World = WorldActorsPair.Key;
+		const TSet<AActor*>& ActorsPendingMovementUpdate = WorldActorsPair.Value;
+
+		if (!IsValid(World))
+		{
+			continue;
+		}
+
+		TMap<AInstancedFoliageActor*, TArray<UActorComponent*>> UpdatedIFAsPerComponent;
+
+		// Process new instance positions first
+		for (TActorIterator<AInstancedFoliageActor> It(World); It; ++It)
+		{
+			for (AActor* Actor : ActorsPendingMovementUpdate)
+			{
+				if (!IsValid(Actor))
+				{
+					continue;
+				}
+
+				AInstancedFoliageActor* IFA = *It;
+				for (UActorComponent* Component : Actor->GetComponents())
+				{
+					if (IFA->MoveInstancesForMovedComponent(Component))
+					{
+						UpdatedIFAsPerComponent.FindOrAdd(IFA).Add(Component);
+					}
+				}
+
+				if (FFoliageHelper::IsOwnedByFoliage(Actor))
+				{
+					IFA->MoveInstancesForMovedOwnedActors(Actor);
+				}
+			}
+		}
+
+		// Then update partitioning
+		for (const TPair<AInstancedFoliageActor*, TArray<UActorComponent*>>& IFAComponentsPair : UpdatedIFAsPerComponent)
+		{
+			for (UActorComponent* Component : IFAComponentsPair.Value)
+			{
+				IFAComponentsPair.Key->UpdateInstancePartitioningForMovedComponent(Component);
+			}
+		}
+	}
+
+	ActorsPendingMovementUpdatePerWorld.Empty();
+}
+
+ETickableTickType UFoliageEditorSubsystem::GetTickableTickType() const
+{
+	return (IsTemplate() ? ETickableTickType::Never : ETickableTickType::Conditional);
+}
+
+bool UFoliageEditorSubsystem::IsTickable() const
+{
+	return ActorsPendingMovementUpdatePerWorld.Num() > 0;
+}
+
+bool UFoliageEditorSubsystem::IsTickableInEditor() const
+{
+	return true;
+}
+
+UWorld* UFoliageEditorSubsystem::GetTickableGameObjectWorld() const
+{
+	return GEditor->GetEditorWorldContext().World();
+}
+
+TStatId UFoliageEditorSubsystem::GetStatId() const
+{
+	RETURN_QUICK_DECLARE_CYCLE_STAT(UFoliageEditorSubsystem, STATGROUP_Tickables);
+}
+
 void UFoliageEditorSubsystem::OnPostApplyLevelOffset(ULevel* InLevel, UWorld* InWorld, const FVector& InOffset, bool bWorldShift)
 {
 	if (InWorld && !InWorld->IsGameWorld())
@@ -90,34 +169,8 @@ void UFoliageEditorSubsystem::OnActorMoved(AActor* InActor)
 		return;
 	}
 
-	TMap<AInstancedFoliageActor*, TArray<UActorComponent*>> UpdatedIFAsPerComponent;
-	
-	// Process new instance positions first
-	for (TActorIterator<AInstancedFoliageActor> It(InWorld); It; ++It)
-	{
-		AInstancedFoliageActor* IFA = *It;
-		for (UActorComponent* Component : InActor->GetComponents())
-		{
-			if (IFA->MoveInstancesForMovedComponent(Component))
-			{
-				UpdatedIFAsPerComponent.FindOrAdd(IFA).Add(Component);
-			}
-		}
-
-		if (FFoliageHelper::IsOwnedByFoliage(InActor))
-		{
-			IFA->MoveInstancesForMovedOwnedActors(InActor);
-		}
-	}
-
-	// Then update partitioning
-	for (const auto& [IFA, Components] : UpdatedIFAsPerComponent)
-	{
-		for (UActorComponent* Component : Components)
-		{
-			IFA->UpdateInstancePartitioningForMovedComponent(Component);
-		}
-	}
+	// Defer movement update until the next tick to reduce the number of expensive TActorIterator operations
+	ActorsPendingMovementUpdatePerWorld.FindOrAdd(InWorld).Add(InActor);
 }
 
 void UFoliageEditorSubsystem::OnActorOuterChanged(AActor* InActor, UObject* OldOuter)
diff --git a/Engine/Source/Editor/FoliageEdit/Public/FoliageEditorSubsystem.h b/Engine/Source/Editor/FoliageEdit/Public/FoliageEditorSubsystem.h
--- a/Engine/Source/Editor/FoliageEdit/Public/FoliageEditorSubsystem.h
+++ b/Engine/Source/Editor/FoliageEdit/Public/FoliageEditorSubsystem.h
@@ -10,7 +10,7 @@ class ULevel;
 class AActor;
 
 UCLASS()
-class UFoliageEditorSubsystem : public UEditorSubsystem
+class UFoliageEditorSubsystem : public UEditorSubsystem, public FTickableGameObject
 {
 	GENERATED_BODY()
 
@@ -19,6 +19,15 @@ public:
 	virtual void Initialize(FSubsystemCollectionBase& Collection) override;
 	virtual void Deinitialize() override;
 
+	//~ Begin FTickableGameObject interface
+	virtual void Tick(float DeltaTime) override;
+	virtual ETickableTickType GetTickableTickType() const override;
+	virtual bool IsTickable() const override;
+	virtual bool IsTickableInEditor() const override;
+	virtual UWorld* GetTickableGameObjectWorld() const override;
+	virtual TStatId GetStatId() const override;
+	//~ End FTickableGameObject interface
+
 private:
 	void OnActorMoved(AActor* InActor);
 	void OnActorOuterChanged(AActor* InActor, UObject* OldOuter);
@@ -26,4 +35,6 @@ private:
 	void OnPostApplyLevelOffset(ULevel* InLevel, UWorld* InWorld, const FVector& InOffset, bool bWorldShift);
 	void OnPostApplyLevelTransform(ULevel* InLevel, const FTransform& InTransform);
 	void OnPostWorldInitialization(UWorld* InWorld, const UWorld::InitializationValues IVS);
+
+	TMap<const UWorld*, TSet<AActor*>> ActorsPendingMovementUpdatePerWorld;
 };
