Fix non-asset foliage types disappearing in OFPA levels (non-WP) due to checking their outermost against the outermost of the painted level instead. Now simply checks the level.

diff --git a/Engine/Source/Editor/FoliageEdit/Private/FoliageEdMode.cpp b/Engine/Source/Editor/FoliageEdit/Private/FoliageEdMode.cpp
--- a/Engine/Source/Editor/FoliageEdit/Private/FoliageEdMode.cpp
+++ b/Engine/Source/Editor/FoliageEdit/Private/FoliageEdMode.cpp
@@ -3382,7 +3382,7 @@ bool FEdModeFoliage::CanPaint(const UFoliageType* FoliageType, const ULevel* InL
 
 	// Non-shared objects can be painted only into their own level
 	// Assets can be painted everywhere
-	if (FoliageType->IsAsset() || FoliageType->GetOutermost() == InLevel->GetOutermost())
+	if (FoliageType->IsAsset() || FoliageType->GetTypedOuter<ULevel>() == InLevel)
 	{
 		return true;
 	}
