[FoliageEditorSubsystem] OnActorMoved optimization to avoid expensive TActorIterator operations when no work is done inside the TActorIterator loop

Both MoveInstancesForMovedComponent and MoveInstancesForMovedOwnedActors (invoked in the TActorIterator loop) early-out when FLevelUtils::IsApplyingLevelTransform() returns true. Take advantage of this and check IsApplyingLevelTransform() before the TActorIterator loop.

diff --git a/Engine/Source/Editor/FoliageEdit/Private/FoliageEditorSubsystem.cpp b/Engine/Source/Editor/FoliageEdit/Private/FoliageEditorSubsystem.cpp
--- a/Engine/Source/Editor/FoliageEdit/Private/FoliageEditorSubsystem.cpp
+++ b/Engine/Source/Editor/FoliageEdit/Private/FoliageEditorSubsystem.cpp
@@ -6,6 +6,7 @@
 #include "Engine/World.h"
 #include "EngineUtils.h"
 #include "InstancedFoliageActor.h"
+#include "LevelUtils.h"
 #include "FoliageHelper.h"
 #include "ActorPartition/ActorPartitionSubsystem.h"
 
@@ -90,6 +91,13 @@ void UFoliageEditorSubsystem::OnActorMoved(AActor* InActor)
 		return;
 	}
 
+	// Both MoveInstancesForMovedComponent and MoveInstancesForMovedOwnedActors early-out when FLevelUtils::IsApplyingLevelTransform() returns true, doing no work.
+	// Add the same check here to avoid expensive TActorIterator operations.
+	if (FLevelUtils::IsApplyingLevelTransform())
+	{
+		return;
+	}
+
 	TMap<AInstancedFoliageActor*, TArray<UActorComponent*>> UpdatedIFAsPerComponent;
 	
 	// Process new instance positions first
