[FoliageEdit] Fix unreachable code error
Removing RefreshPalette and related code as well, since it was added in the same change as the unreachable code (back in 2016) and only used in the unreachable code section which was never executed anyway.

diff --git a/Engine/Source/Editor/FoliageEdit/Private/SFoliagePalette.cpp b/Engine/Source/Editor/FoliageEdit/Private/SFoliagePalette.cpp
--- a/Engine/Source/Editor/FoliageEdit/Private/SFoliagePalette.cpp
+++ b/Engine/Source/Editor/FoliageEdit/Private/SFoliagePalette.cpp
@@ -216,7 +216,6 @@ void SFoliagePalette::Construct(const FArguments& InArgs)
 	bItemsNeedRebuild = false;
 	bIsUneditableFoliageTypeSelected = false;
 	bIsRebuildTimerRegistered = false;
-	bIsRefreshTimerRegistered = false;
 
 	FoliageEditMode = InArgs._FoliageEdMode;
 
@@ -450,16 +449,6 @@ void SFoliagePalette::UpdatePalette(bool bRebuildItems)
 	}
 }
 
-void SFoliagePalette::RefreshPalette()
-{
-	// Do not register the refresh timer if we're pending a rebuild; rebuild should cause the palette to refresh
-	if (!bIsRefreshTimerRegistered && !bIsRebuildTimerRegistered)
-	{
-		bIsRefreshTimerRegistered = true;
-		RegisterActiveTimer(0.f, FWidgetActiveTimerDelegate::CreateSP(this, &SFoliagePalette::RefreshPaletteItems));
-	}
-}
-
 void SFoliagePalette::UpdateThumbnailForType(UFoliageType* FoliageType)
 {
 	// Recreate the palette item for the given foliage type
@@ -1329,9 +1318,6 @@ bool SFoliagePalette::CanSelectInstances() const
 TSharedRef<ITableRow> SFoliagePalette::GenerateTile(FFoliagePaletteItemModelPtr Item, const TSharedRef<STableViewBase>& OwnerTable)
 {
 	return SNew(SFoliagePaletteItemTile, OwnerTable, Item);
-
-	// Refresh the palette to ensure that thumbnails are correct
-	RefreshPalette();
 }
 
 float SFoliagePalette::GetScaledThumbnailSize() const
@@ -1562,16 +1548,4 @@ EActiveTimerReturnType SFoliagePalette::UpdatePaletteItems(double InCurrentTime,
 	return EActiveTimerReturnType::Stop;
 }
 
-EActiveTimerReturnType SFoliagePalette::RefreshPaletteItems(double InCurrentTime, float InDeltaTime)
-{
-	// Do not refresh the palette if we're waiting on a rebuild
-	if (!bItemsNeedRebuild)
-	{
-		RefreshActivePaletteViewWidget();
-	}
-
-	bIsRefreshTimerRegistered = false;
-	return EActiveTimerReturnType::Stop;
-}
-
 #undef LOCTEXT_NAMESPACE
diff --git a/Engine/Source/Editor/FoliageEdit/Private/SFoliagePalette.h b/Engine/Source/Editor/FoliageEdit/Private/SFoliagePalette.h
--- a/Engine/Source/Editor/FoliageEdit/Private/SFoliagePalette.h
+++ b/Engine/Source/Editor/FoliageEdit/Private/SFoliagePalette.h
@@ -71,9 +71,6 @@ public:
 	/** Updates the foliage palette, optionally doing a full rebuild of the items in the palette as well */
 	void UpdatePalette(bool bRebuildItems = false);
 
-	/** Refreshes the foliage palette */
-	void RefreshPalette();
-
 	/** Updates the thumbnail for the given foliage type in the palette */
 	void UpdateThumbnailForType(UFoliageType* FoliageType);
 
@@ -283,9 +280,6 @@ private:
 	/** Active timer handler to update the items in the palette */
 	EActiveTimerReturnType UpdatePaletteItems(double InCurrentTime, float InDeltaTime);
 
-	/** Active timer handler to refresh the palette */
-	EActiveTimerReturnType RefreshPaletteItems(double InCurrentTime, float InDeltaTime);
-
 private:
 	typedef TTextFilter<FFoliagePaletteItemModelPtr> FoliageTypeTextFilter;
 	TSharedPtr<FoliageTypeTextFilter> TypeFilter;
@@ -331,5 +325,4 @@ private:
 	bool bIsUneditableFoliageTypeSelected : 1;
 	
 	bool bIsRebuildTimerRegistered : 1;
-	bool bIsRefreshTimerRegistered : 1;
 };
