[Foliage]

Expose an API which allows external systems to register component classes which should never be used as "base" components for foliage. This means that any foliage spawned on those components are not directly tied to their lifecycle or transforms.

diff --git a/Engine/Source/Editor/FoliageEdit/Private/FoliageEditModule.cpp b/Engine/Source/Editor/FoliageEdit/Private/FoliageEditModule.cpp
--- a/Engine/Source/Editor/FoliageEdit/Private/FoliageEditModule.cpp
+++ b/Engine/Source/Editor/FoliageEdit/Private/FoliageEditModule.cpp
@@ -86,6 +86,8 @@ public:
 		// Register thumbnail renderer
 		UThumbnailManager::Get().RegisterCustomRenderer(UFoliageType_InstancedStaticMesh::StaticClass(), UFoliageType_ISMThumbnailRenderer::StaticClass());
 		UThumbnailManager::Get().RegisterCustomRenderer(UFoliageType_Actor::StaticClass(), UFoliageType_ActorThumbnailRenderer::StaticClass());
+
+		IFoliageEditModuleBase::SetFoliageEditModulePtr(this);
 	}
 
 	/**
@@ -93,6 +95,8 @@ public:
 	 */
 	virtual void ShutdownModule() override
 	{
+		IFoliageEditModuleBase::SetFoliageEditModulePtr(nullptr);
+
 		if (GUnrealEd)
 		{
 			GUnrealEd->UnregisterComponentVisualizer(UProceduralFoliageComponent::StaticClass()->GetFName());
@@ -245,8 +249,35 @@ public:
 		return FoliageMode->CanMoveSelectedFoliageToLevel(InTargetLevel);
 	}
 
+	virtual void RegisterComponentBaseIDClassToIgnore(TSubclassOf<UActorComponent> InClassToIgnore) override
+	{
+		ComponentBaseIDClassesToIgnore.Add(InClassToIgnore.Get());
+	}
+
+	virtual void UnregisterComponentBaseIDClassToIgnore(TSubclassOf<UActorComponent> InClassToUnregister) override
+	{
+		ComponentBaseIDClassesToIgnore.Remove(InClassToUnregister.Get());
+	}
+
+	virtual bool ShouldIgnoreComponentForBaseID(UActorComponent* InComponent) const
+	{
+		for (TSoftObjectPtr<UClass> ClassToCheckPtr : ComponentBaseIDClassesToIgnore)
+		{
+			if (UClass* ClassToCheck = ClassToCheckPtr.Get())
+			{
+				if (InComponent->IsA(ClassToCheck))
+				{
+					return true;
+				}
+			}
+		}
+		return false;
+	}
+
 	FDelegateHandle OnLevelActorDeletedDelegateHandle;
 	FDelegateHandle OnExperimentalSettingChangedDelegateHandle;
+
+	TSet<TSoftObjectPtr<UClass>> ComponentBaseIDClassesToIgnore;
 #endif
 };
 
diff --git a/Engine/Source/Editor/FoliageEdit/Public/FoliageEditModule.h b/Engine/Source/Editor/FoliageEdit/Public/FoliageEditModule.h
--- a/Engine/Source/Editor/FoliageEdit/Public/FoliageEditModule.h
+++ b/Engine/Source/Editor/FoliageEdit/Public/FoliageEditModule.h
@@ -5,6 +5,7 @@
 #include "CoreMinimal.h"
 #include "Modules/ModuleInterface.h"
 #include "UObject/NameTypes.h"
+#include "FoliageModule.h"
 
 extern const FName FoliageEditAppIdentifier;
 
@@ -13,15 +14,7 @@ class ULevel;
 /**
  * Foliage Edit mode module interface
  */
-class IFoliageEditModule : public IModuleInterface
+class IFoliageEditModule : public IFoliageEditModuleBase
 {
 public:
-
-#if WITH_EDITOR
-	/** Move the selected foliage to the specified level */
-	virtual void MoveSelectedFoliageToLevel(ULevel* InTargetLevel) = 0;
-	/** Notifies us that the foliage has been externally changed and needs refreshing.  */
-	virtual void UpdateMeshList() = 0;
-	virtual bool CanMoveSelectedFoliageToLevel(ULevel* InTargetLevel) const = 0;
-#endif
 };
diff --git a/Engine/Source/Runtime/Foliage/Private/FoliageInstanceBase.cpp b/Engine/Source/Runtime/Foliage/Private/FoliageInstanceBase.cpp
--- a/Engine/Source/Runtime/Foliage/Private/FoliageInstanceBase.cpp
+++ b/Engine/Source/Runtime/Foliage/Private/FoliageInstanceBase.cpp
@@ -7,6 +7,8 @@
 #include "Engine/Level.h"
 #include "Engine/WorldComposition.h"
 #include "UObject/FortniteMainBranchObjectVersion.h"
+#include "Modules/ModuleManager.h"
+#include "FoliageModule.h"
 
 #if WITH_EDITORONLY_DATA
 
@@ -124,7 +126,15 @@ FArchive& operator << (FArchive& Ar, FFoliageInstanceBaseCache& InstanceBaseCach
 FFoliageInstanceBaseId FFoliageInstanceBaseCache::AddInstanceBaseId(UActorComponent* InComponent)
 {
 	FFoliageInstanceBaseId BaseId = FFoliageInstanceBaseCache::InvalidBaseId;
-	if (InComponent && !InComponent->IsCreatedByConstructionScript())
+	bool bIgnoreComponentClass = false; 
+#if WITH_EDITOR
+	if (IFoliageEditModuleBase* FoliageModule = IFoliageEditModuleBase::Get(); ensure(FoliageModule))
+	{
+		bIgnoreComponentClass = FoliageModule->ShouldIgnoreComponentForBaseID(InComponent);
+	}
+#endif // WITH_EDITOR
+
+	if (InComponent && !InComponent->IsCreatedByConstructionScript() && !bIgnoreComponentClass)
 	{
 		BaseId = GetInstanceBaseId(InComponent);
 		FFoliageInstanceBaseInfo BaseInfo(InComponent);
diff --git a/Engine/Source/Runtime/Foliage/Private/FoliageModule.cpp b/Engine/Source/Runtime/Foliage/Private/FoliageModule.cpp
--- a/Engine/Source/Runtime/Foliage/Private/FoliageModule.cpp
+++ b/Engine/Source/Runtime/Foliage/Private/FoliageModule.cpp
@@ -4,3 +4,17 @@
 #include "Modules/ModuleManager.h"
 
 IMPLEMENT_MODULE( IFoliageModule, Foliage );
+
+#if WITH_EDITOR
+IFoliageEditModuleBase* IFoliageEditModuleBase::FoliageEditModulePtr = nullptr;
+
+IFoliageEditModuleBase* IFoliageEditModuleBase::Get()
+{
+	return FoliageEditModulePtr;
+}
+
+void IFoliageEditModuleBase::SetFoliageEditModulePtr(IFoliageEditModuleBase* InFoliageEditModulePtr)
+{
+	FoliageEditModulePtr = InFoliageEditModulePtr;
+}
+#endif // WITH_EDITOR
\ No newline at end of file
diff --git a/Engine/Source/Runtime/Foliage/Public/FoliageModule.h b/Engine/Source/Runtime/Foliage/Public/FoliageModule.h
--- a/Engine/Source/Runtime/Foliage/Public/FoliageModule.h
+++ b/Engine/Source/Runtime/Foliage/Public/FoliageModule.h
@@ -4,10 +4,43 @@
 
 #include "CoreMinimal.h"
 #include "Modules/ModuleInterface.h"
+#include "Templates/SubclassOf.h"
+
+class ULevel;
+class UActorComponent;
 
 /**
  * Foliage Edit mode module interface
  */
+class IFoliageEditModuleBase : public IModuleInterface
+{
+public:
+#if WITH_EDITOR
+	/** Move the selected foliage to the specified level */
+	virtual void MoveSelectedFoliageToLevel(ULevel* InTargetLevel) = 0;
+	/** Notifies us that the foliage has been externally changed and needs refreshing.  */
+	virtual void UpdateMeshList() = 0;
+	virtual bool CanMoveSelectedFoliageToLevel(ULevel* InTargetLevel) const = 0;
+
+	virtual bool ShouldIgnoreComponentForBaseID(UActorComponent* InComponent) const = 0;
+	virtual void RegisterComponentBaseIDClassToIgnore(TSubclassOf<UActorComponent> InClassToIgnore) = 0;
+	virtual void UnregisterComponentBaseIDClassToIgnore(TSubclassOf<UActorComponent> InClassToUnregister) = 0;
+
+	FOLIAGE_API static IFoliageEditModuleBase* Get();
+#endif
+	
+protected:
+
+#if WITH_EDITOR
+	FOLIAGE_API static void SetFoliageEditModulePtr(IFoliageEditModuleBase* InFoliageEditModulePtr);
+#endif
+
+private:
+#if WITH_EDITOR
+	static IFoliageEditModuleBase* FoliageEditModulePtr;
+#endif
+};
+
 class IFoliageModule : public IModuleInterface
 {
 public:
